[CrashFix] Stack overflow in IDO fixup tool when fixing objects with circular instanced object properties

diff --git a/Engine/Plugins/Experimental/InstanceDataObjectFixupTool/Source/Private/InstanceDataObjectFixupPanel.cpp b/Engine/Plugins/Experimental/InstanceDataObjectFixupTool/Source/Private/InstanceDataObjectFixupPanel.cpp
--- a/Engine/Plugins/Experimental/InstanceDataObjectFixupTool/Source/Private/InstanceDataObjectFixupPanel.cpp
+++ b/Engine/Plugins/Experimental/InstanceDataObjectFixupTool/Source/Private/InstanceDataObjectFixupPanel.cpp
@@ -626,8 +626,8 @@ void FInstanceDataObjectFixupPanel::OnRedirectProperty(FPropertyPath From, FProp
 	RedirectProperty(From, To);
 }
 
-static void InitRedirectedPropertyTreeRec(const TSharedPtr<FRedirectedPropertyNode>& Node, FProperty* Property, void* Value);
-static void InitRedirectedPropertyTreeRec(const TSharedPtr<FRedirectedPropertyNode>& Node, UStruct* Struct, void* StructValue)
+static void InitRedirectedPropertyTreeRec(const TSharedPtr<FRedirectedPropertyNode>& Node, FProperty* Property, void* Value, TSet<UObject*>& EnteredObjects);
+static void InitRedirectedPropertyTreeRec(const TSharedPtr<FRedirectedPropertyNode>& Node, UStruct* Struct, void* StructValue, TSet<UObject*>& EnteredObjects)
 {
 	for (FProperty* Property : TFieldRange<FProperty>(Struct))
 	{
@@ -637,7 +637,7 @@ static void InitRedirectedPropertyTreeRec(const TSharedPtr<FRedirectedPropertyNo
 			{
 				const TSharedPtr<FRedirectedPropertyNode>& ChildNode = Node->FindOrAdd(FPropertyInfo(Property));
 				void* Value = Property->ContainerPtrToValuePtr<void>(StructValue);
-				InitRedirectedPropertyTreeRec(ChildNode, Property, Value);
+				InitRedirectedPropertyTreeRec(ChildNode, Property, Value, EnteredObjects);
 			}
 		}
 		else
@@ -648,18 +648,18 @@ static void InitRedirectedPropertyTreeRec(const TSharedPtr<FRedirectedPropertyNo
             	{
             		const TSharedPtr<FRedirectedPropertyNode>& ChildNode = Node->FindOrAdd(FPropertyInfo(Property, StaticArrayIndex));
             		void* Value = Property->ContainerPtrToValuePtr<void>(StructValue, StaticArrayIndex);
-            		InitRedirectedPropertyTreeRec(ChildNode, Property, Value);
+            		InitRedirectedPropertyTreeRec(ChildNode, Property, Value, EnteredObjects);
             	}
             }
 		}
 	}
 }
 
-static void InitRedirectedPropertyTreeRec(const TSharedPtr<FRedirectedPropertyNode>& Node, FProperty* Property, void* Value)
+static void InitRedirectedPropertyTreeRec(const TSharedPtr<FRedirectedPropertyNode>& Node, FProperty* Property, void* Value, TSet<UObject*>& EnteredObjects)
 {
 	if (const FStructProperty* AsStructProperty = CastField<FStructProperty>(Property))
 	{
-		InitRedirectedPropertyTreeRec(Node, AsStructProperty->Struct, Value);
+		InitRedirectedPropertyTreeRec(Node, AsStructProperty->Struct, Value, EnteredObjects);
 	}
 	else if (const FObjectProperty* AsObjectProperty = CastField<FObjectProperty>(Property))
 	{
@@ -672,7 +672,13 @@ static void InitRedirectedPropertyTreeRec(const TSharedPtr<FRedirectedPropertyNo
 				{
 					Object = Found;
 				}
-            	InitRedirectedPropertyTreeRec(Node, Object->GetClass(), Object);
+				// check for circular references to avoid infinite recursion
+				if (!EnteredObjects.Contains(Object))
+				{
+					EnteredObjects.Add(Object);
+            		InitRedirectedPropertyTreeRec(Node, Object->GetClass(), Object, EnteredObjects);
+					EnteredObjects.Remove(Object);
+            	}
             }
 		}
 		
@@ -683,7 +689,7 @@ static void InitRedirectedPropertyTreeRec(const TSharedPtr<FRedirectedPropertyNo
 		for (int32 ArrayIndex = 0; ArrayIndex < Array.Num(); ++ArrayIndex)
 		{
 			const TSharedPtr<FRedirectedPropertyNode>& ChildNode = Node->FindOrAdd(FPropertyInfo(AsArrayProperty->Inner, ArrayIndex));
-			InitRedirectedPropertyTreeRec(ChildNode, AsArrayProperty->Inner, Array.GetElementPtr(ArrayIndex));
+			InitRedirectedPropertyTreeRec(ChildNode, AsArrayProperty->Inner, Array.GetElementPtr(ArrayIndex), EnteredObjects);
 		}
 	}
 	else if (const FSetProperty* AsSetProperty = CastField<FSetProperty>(Property))
@@ -692,7 +698,7 @@ static void InitRedirectedPropertyTreeRec(const TSharedPtr<FRedirectedPropertyNo
 		for (FScriptSetHelper::FIterator Itr = Set.CreateIterator(); Itr; ++Itr)
 		{
 			const TSharedPtr<FRedirectedPropertyNode>& ChildNode = Node->FindOrAdd(FPropertyInfo(AsSetProperty->ElementProp, Itr.GetLogicalIndex()));
-			InitRedirectedPropertyTreeRec(ChildNode, AsSetProperty->ElementProp, Set.GetElementPtr(Itr));
+			InitRedirectedPropertyTreeRec(ChildNode, AsSetProperty->ElementProp, Set.GetElementPtr(Itr), EnteredObjects);
 		}
 	}
 	else if (const FMapProperty* AsMapProperty = CastField<FMapProperty>(Property))
@@ -701,16 +707,19 @@ static void InitRedirectedPropertyTreeRec(const TSharedPtr<FRedirectedPropertyNo
 		for (FScriptMapHelper::FIterator Itr = Map.CreateIterator(); Itr; ++Itr)
 		{
 			const TSharedPtr<FRedirectedPropertyNode>& KeyNode = Node->FindOrAdd(FPropertyInfo(AsMapProperty->KeyProp, Itr.GetLogicalIndex()));
-			InitRedirectedPropertyTreeRec(KeyNode, AsMapProperty->KeyProp, Map.GetKeyPtr(Itr));
+			InitRedirectedPropertyTreeRec(KeyNode, AsMapProperty->KeyProp, Map.GetKeyPtr(Itr), EnteredObjects);
 			const TSharedPtr<FRedirectedPropertyNode>& ValNode = Node->FindOrAdd(FPropertyInfo(AsMapProperty->ValueProp, Itr.GetLogicalIndex()));
-			InitRedirectedPropertyTreeRec(ValNode, AsMapProperty->ValueProp, Map.GetValuePtr(Itr));
+			InitRedirectedPropertyTreeRec(ValNode, AsMapProperty->ValueProp, Map.GetValuePtr(Itr), EnteredObjects);
 		}
 	}
 }
 
 void FInstanceDataObjectFixupPanel::InitRedirectedPropertyTree()
 {
-	InitRedirectedPropertyTreeRec(RedirectedPropertyTree, Instances[0]->GetClass(), Instances[0]);
+	TSet<UObject*> EnteredObjects = {Instances[0]};
+	InitRedirectedPropertyTreeRec(RedirectedPropertyTree, Instances[0]->GetClass(), Instances[0], EnteredObjects);
+	EnteredObjects.Remove(Instances[0]);
+	check(EnteredObjects.IsEmpty());
 }
 
 void UInstanceDataObjectFixupUndoHandler::Init(const TSharedRef<FInstanceDataObjectFixupPanel>& Panel)
