[Feature] using FDetailsNameWidgetOverrideCustomization to customize name widgets in IDO fixup tool.

diff --git a/Engine/Plugins/Experimental/InstanceDataObjectFixupTool/Source/Private/InstanceDataObjectFixupDetailCustomization.cpp b/Engine/Plugins/Experimental/InstanceDataObjectFixupTool/Source/Private/InstanceDataObjectFixupDetailCustomization.cpp
--- a/Engine/Plugins/Experimental/InstanceDataObjectFixupTool/Source/Private/InstanceDataObjectFixupDetailCustomization.cpp
+++ b/Engine/Plugins/Experimental/InstanceDataObjectFixupTool/Source/Private/InstanceDataObjectFixupDetailCustomization.cpp
@@ -14,93 +14,6 @@
 
 #define LOCTEXT_NAMESPACE "InstanceDataObjectFixupDetails"
 
-
-TSet<FPropertyPath> FInstanceDataObjectFixupDetailNodeBuilder::GetRedirectOptions(const UStruct* Struct, void* Value) const
-{
-	TSet<FPropertyPath> Result;
-	GetRedirectOptions(Struct, Value, {}, Result);
-	return Result;
-}
-
-void FInstanceDataObjectFixupDetailNodeBuilder::GetRedirectOptions(const UStruct* Struct, void* Value, const FPropertyPath& Path, TSet<FPropertyPath>& OutPaths) const
-{
-	for (FProperty* SubProperty : TFieldRange<FProperty>(Struct))
-    {
-		if (SubProperty->ArrayDim == 1)
-		{
-			const FPropertyPath SubPath = Path.ExtendPath(FPropertyInfo(SubProperty)).Get();
-			GetRedirectOptions(SubProperty, SubProperty->ContainerPtrToValuePtr<void>(Value), SubPath, OutPaths);
-		}
-		else
-		{
-			for (int32 ArrayIndex = 0; ArrayIndex < SubProperty->ArrayDim; ++ArrayIndex)
-			{
-    			const FPropertyPath SubPath = Path.ExtendPath(FPropertyInfo(SubProperty, ArrayIndex)).Get();
-    			GetRedirectOptions(SubProperty, SubProperty->ContainerPtrToValuePtr<void>(Value, ArrayIndex), SubPath, OutPaths);
-			}
-		}
-    }
-}
-
-void FInstanceDataObjectFixupDetailNodeBuilder::GetRedirectOptions(const FProperty* Property, void* Value, const FPropertyPath& Path, TSet<FPropertyPath>& OutPaths) const
-{
-	if (Property->GetBoolMetaData(TEXT("isLoose")))
-	{
-		// don't include loose properties as options
-		return;
-	}
-
-	if (!DiffPanel.Pin()->RedirectedPropertyTree->Find(Path))
-	{
-		OutPaths.Add(Path);
-	}
-	
-	if (const FStructProperty* AsStructProperty = CastField<FStructProperty>(Property))
-	{
-		GetRedirectOptions(AsStructProperty->Struct, Value, Path, OutPaths);
-	}
-	if (const FObjectProperty* AsObjectProperty = CastField<FObjectProperty>(Property))
-	{
-		if (AsObjectProperty->HasAnyPropertyFlags(CPF_InstancedReference))
-		{
-			TObjectPtr<UObject> Object = AsObjectProperty->GetObjectPropertyValue(Value);
-			UE::FPropertyBagRepository& PropertyBagRepository = UE::FPropertyBagRepository::Get();
-			if (UObject* Found = PropertyBagRepository.FindInstanceDataObject(Object))
-			{
-				Object = Found;
-			}
-			GetRedirectOptions(Object->GetClass(), Object, Path, OutPaths);
-		}
-	}
-	else if (const FArrayProperty* AsArrayProperty = CastField<FArrayProperty>(Property))
-	{
-		FScriptArrayHelper Array(AsArrayProperty, Value);
-		for (int32 ArrayIndex = 0; ArrayIndex < Array.Num(); ++ArrayIndex)
-		{
-			const FPropertyPath SubPath = Path.ExtendPath(FPropertyInfo(AsArrayProperty->Inner, ArrayIndex)).Get();
-			GetRedirectOptions(AsArrayProperty->Inner, Array.GetElementPtr(ArrayIndex), SubPath, OutPaths);
-		}
-	}
-	else if (const FSetProperty* AsSetProperty = CastField<FSetProperty>(Property))
-	{
-		FScriptSetHelper Set(AsSetProperty, Value);
-		for (FScriptSetHelper::FIterator Iterator = Set.CreateIterator(); Iterator; ++Iterator)
-		{
-			const FPropertyPath SubPath = Path.ExtendPath(FPropertyInfo(AsSetProperty->ElementProp, Iterator.GetLogicalIndex())).Get();
-			GetRedirectOptions(AsSetProperty->ElementProp, Set.GetElementPtr(Iterator), SubPath, OutPaths);
-		}
-	}
-	else if (const FMapProperty* AsMapProperty = CastField<FMapProperty>(Property))
-	{
-		FScriptMapHelper Map(AsMapProperty, Value);
-		for (FScriptMapHelper::FIterator Iterator = Map.CreateIterator(); Iterator; ++Iterator)
-		{
-			const FPropertyPath SubPath = Path.ExtendPath(FPropertyInfo(AsMapProperty->ValueProp, Iterator.GetLogicalIndex())).Get();
-			GetRedirectOptions(AsMapProperty->ValueProp, Map.GetValuePtr(Iterator), SubPath, OutPaths);
-		}
-	}
-}
-
 /////////////////////////////////////////////////////////////////
 // FInstanceDataObjectFixupDetailCustomization
 /////////////////////////////////////////////////////////////////
@@ -134,56 +47,45 @@ void FInstanceDataObjectFixupDetailCustomization::CustomizeDetails(IDetailLayout
 	for (const FName CategoryName : CategoryNames)
 	{
 		IDetailCategoryBuilder& Category = DetailBuilder.EditCategory(CategoryName);
-		TArray<TSharedRef<IPropertyHandle>> NonAdvancedProperties;
-		Category.GetDefaultProperties(NonAdvancedProperties, true, false);
-		for (const TSharedRef<IPropertyHandle>& Handle : NonAdvancedProperties)
-		{
-			DetailBuilder.HideProperty(Handle);
-			Category.AddCustomBuilder(MakeShared<FInstanceDataObjectFixupDetailNodeBuilder>(Panel.ToSharedRef(), Handle), false);
-		}
-		TArray<TSharedRef<IPropertyHandle>> AdvancedProperties;
-		Category.GetDefaultProperties(AdvancedProperties, false, true);
-		for (const TSharedRef<IPropertyHandle>& Handle : AdvancedProperties)
+		TArray<TSharedRef<IPropertyHandle>> Handles;
+		Category.GetDefaultProperties(Handles, true, true);
+		for (const TSharedRef<IPropertyHandle>& Handle : Handles)
 		{
-			DetailBuilder.HideProperty(Handle);
-			Category.AddCustomBuilder(MakeShared<FInstanceDataObjectFixupDetailNodeBuilder>(Panel.ToSharedRef(), Handle), true);
+			CustomizeHandle(Handle, DetailBuilder);
 		}
 	}
 }
 
-/////////////////////////////////////////////////////////////////
-// FHideLoosePropertiesCustomization
-/////////////////////////////////////////////////////////////////
-
-FHideLoosePropertiesCustomization::~FHideLoosePropertiesCustomization()
-{
-}
-
-void FHideLoosePropertiesCustomization::CustomizeDetails(IDetailLayoutBuilder& DetailBuilder)
+bool FInstanceDataObjectFixupDetailCustomization::IsHidden(const TSharedPtr<IPropertyHandle>& PropertyHandle) const
 {
-	TArray<TWeakObjectPtr<UObject>> Objects;
-	DetailBuilder.GetObjectsBeingCustomized(Objects);
-	if (Objects.IsEmpty())
+	const TSharedPtr<FInstanceDataObjectFixupPanel> Panel = DiffPanel.Pin();
+	if (!Panel)
 	{
-		return;
+		return true;
 	}
-	TArray<FName> CategoryNames;
-	DetailBuilder.GetCategoryNames(CategoryNames);
-	for (const FName CategoryName : CategoryNames)
+
+	if (const FProperty* Property = PropertyHandle->GetProperty())
 	{
-		IDetailCategoryBuilder& Category = DetailBuilder.EditCategory(CategoryName);
-		TArray<TSharedRef<IPropertyHandle>> Handles;
-		Category.GetDefaultProperties(Handles, true, true);
-		for (const TSharedRef<IPropertyHandle>& Handle : Handles)
+		if (Panel->HasViewFlag(FInstanceDataObjectFixupPanel::EViewFlags::HideLooseProperties) &&
+			Property->GetBoolMetaData(TEXT("isLoose")))
 		{
-			CustomizeHandle(Handle, DetailBuilder);
+			return true;
+		}
+
+		if (Panel->HasViewFlag(FInstanceDataObjectFixupPanel::EViewFlags::IncludeOnlySetBySerialization))
+		{
+			if (!Panel->IsInRedirectedPropertyTree(*PropertyHandle->CreateFPropertyPath()))
+			{
+				return true;
+			}
 		}
 	}
+	return false;
 }
 
-void FHideLoosePropertiesCustomization::CustomizeHandle(const TSharedRef<IPropertyHandle>& Handle, IDetailLayoutBuilder& DetailBuilder)
+void FInstanceDataObjectFixupDetailCustomization::CustomizeHandle(const TSharedRef<IPropertyHandle>& Handle, IDetailLayoutBuilder& DetailBuilder)
 {
-	if (Handle->GetBoolMetaData(TEXT("isLoose")))
+	if (IsHidden(Handle))
 	{
 		DetailBuilder.HideProperty(Handle);
 	}
@@ -201,56 +103,21 @@ void FHideLoosePropertiesCustomization::CustomizeHandle(const TSharedRef<IProper
 
 
 /////////////////////////////////////////////////////////////////
-// FInstanceDataObjectFixupDetailNodeBuilder
+// FInstanceDataObjectNameWidgetOverride
 /////////////////////////////////////////////////////////////////
 
-FInstanceDataObjectFixupDetailNodeBuilder::FInstanceDataObjectFixupDetailNodeBuilder(const TSharedRef<FInstanceDataObjectFixupPanel>& InDiffPanel, const TSharedRef<IPropertyHandle>& InPropertyHandle)
+FInstanceDataObjectNameWidgetOverride::FInstanceDataObjectNameWidgetOverride(const TSharedRef<FInstanceDataObjectFixupPanel>& InDiffPanel)
 	: DiffPanel(InDiffPanel)
-	, PropertyHandle(InPropertyHandle)
-{}
-
-void FInstanceDataObjectFixupDetailNodeBuilder::GenerateHeaderRowContent(FDetailWidgetRow& NodeRow)
-{	
-	const TSharedPtr<FInstanceDataObjectFixupPanel> Panel = DiffPanel.Pin();
-	if (!Panel)
-	{
-		return;
-	}
-
-	// skip properties that are hidden from this view mode
-	if (IsHidden())
-	{
-		return;
-	}
-
-	// handle categories of sub-objects
-	if (PropertyHandle->IsCategoryHandle())
-	{
-		NodeRow
-		[
-			SNew(SBox)
-			.VAlign(VAlign_Center)
-			[
-				SNew(STextBlock)
-				.Text(PropertyHandle->GetPropertyDisplayName())
-			]
-		];
-		return;
-	}
-	
-	TSharedPtr<SWidget> InnerNameContent = PropertyHandle->CreatePropertyNameWidget();
-	if (PropertyHandle->GetKeyHandle())
-	{
-		// TODO: @jordan.hoffmann this solution for keys only works with basic types. figure out a more robust solution.
-		InnerNameContent = PropertyHandle->GetKeyHandle()->CreatePropertyValueWidget();
-		InnerNameContent->SetEnabled(false);
-	}
+{
+}
 
+TSharedRef<SWidget> FInstanceDataObjectNameWidgetOverride::CustomizeName(TSharedRef<SWidget> InnerNameContent, FPropertyPath& Path)
+{
 	const TSharedRef<SWidget> NameContent = SNew(SWidgetSwitcher)
-		.WidgetIndex(this, &FInstanceDataObjectFixupDetailNodeBuilder::GetNameWidgetIndex)
+		.WidgetIndex(this, &FInstanceDataObjectNameWidgetOverride::GetNameWidgetIndex, Path)
 		+SWidgetSwitcher::Slot()
 		[
-			InnerNameContent.ToSharedRef()
+			InnerNameContent
 		]
 		+SWidgetSwitcher::Slot()
 		[
@@ -260,7 +127,7 @@ void FInstanceDataObjectFixupDetailNodeBuilder::GenerateHeaderRowContent(FDetail
 			.AutoWidth()
 			[
 				SNew(STextBlock)
-				.Visibility(this, &FInstanceDataObjectFixupDetailNodeBuilder::DeletionSymbolVisibility)
+				.Visibility(this, &FInstanceDataObjectNameWidgetOverride::DeletionSymbolVisibility, Path)
 				.Text(LOCTEXT("MarkedForDeletion", "❌"))
 			]
 			+SHorizontalBox::Slot()
@@ -269,90 +136,114 @@ void FInstanceDataObjectFixupDetailNodeBuilder::GenerateHeaderRowContent(FDetail
 			[
 				SNew(SComboButton)
 				.Visibility(EVisibility::Visible)
-				.OnGetMenuContent_Raw(this, &FInstanceDataObjectFixupDetailNodeBuilder::GeneratePropertyRedirectMenu)
+				.OnGetMenuContent_Raw(this, &FInstanceDataObjectNameWidgetOverride::GeneratePropertyRedirectMenu, Path)
 				.ButtonContent()
 				[
-					InnerNameContent.ToSharedRef()
+					InnerNameContent
 				]
 			]
 		];
-	
-	const TSharedRef<SWidget> ValueContent = PropertyHandle->CreatePropertyValueWidget();
-	ValueContent->SetEnabled(!Panel->HasViewFlag(FInstanceDataObjectFixupPanel::EViewFlags::ReadonlyValues));
-	ValueContent->SetVisibility(TAttribute<EVisibility>(this, &FInstanceDataObjectFixupDetailNodeBuilder::ValueContentVisibility));
-	
-	NodeRow
-	.NameContent()
-	[
-		NameContent
-	]
-	.ValueContent()
-	[
-		ValueContent
-	]
-	.PropertyHandleList({PropertyHandle});
+	return NameContent;
 }
 
-void FInstanceDataObjectFixupDetailNodeBuilder::GenerateChildContent(IDetailChildrenBuilder& ChildrenBuilder)
+
+TSet<FPropertyPath> FInstanceDataObjectNameWidgetOverride::GetRedirectOptions(const UStruct* Struct, void* Value) const
 {
-	const TSharedPtr<FInstanceDataObjectFixupPanel> Panel = DiffPanel.Pin();
-	if (!Panel)
+	TSet<FPropertyPath> Result;
+	GetRedirectOptions(Struct, Value, {}, Result);
+	return Result;
+}
+
+void FInstanceDataObjectNameWidgetOverride::GetRedirectOptions(const UStruct* Struct, void* Value, const FPropertyPath& Path, TSet<FPropertyPath>& OutPaths) const
+{
+	for (FProperty* SubProperty : TFieldRange<FProperty>(Struct))
+    {
+		if (SubProperty->ArrayDim == 1)
+		{
+			const FPropertyPath SubPath = Path.ExtendPath(FPropertyInfo(SubProperty)).Get();
+			GetRedirectOptions(SubProperty, SubProperty->ContainerPtrToValuePtr<void>(Value), SubPath, OutPaths);
+		}
+		else
+		{
+			for (int32 ArrayIndex = 0; ArrayIndex < SubProperty->ArrayDim; ++ArrayIndex)
+			{
+    			const FPropertyPath SubPath = Path.ExtendPath(FPropertyInfo(SubProperty, ArrayIndex)).Get();
+    			GetRedirectOptions(SubProperty, SubProperty->ContainerPtrToValuePtr<void>(Value, ArrayIndex), SubPath, OutPaths);
+			}
+		}
+    }
+}
+
+void FInstanceDataObjectNameWidgetOverride::GetRedirectOptions(const FProperty* Property, void* Value, const FPropertyPath& Path, TSet<FPropertyPath>& OutPaths) const
+{
+	if (Property->GetBoolMetaData(TEXT("isLoose")))
 	{
+		// don't include loose properties as options
 		return;
 	}
 
-	// skip properties that are hidden from this view mode
-	if (IsHidden())
+	if (!DiffPanel.Pin()->RedirectedPropertyTree->Find(Path))
 	{
-		return;
+		OutPaths.Add(Path);
 	}
 	
-	uint32 ChildCount;
-	PropertyHandle->GetNumChildren(ChildCount);
-	if (ChildCount == 1 && PropertyHandle->GetProperty() && PropertyHandle->GetProperty()->IsA<FObjectProperty>())
+	if (const FStructProperty* AsStructProperty = CastField<FStructProperty>(Property))
 	{
-		const TSharedRef<IPropertyHandle> ChildHandle = PropertyHandle->GetChildHandle(0).ToSharedRef();
-		if (!ChildHandle->GetProperty())
+		GetRedirectOptions(AsStructProperty->Struct, Value, Path, OutPaths);
+	}
+	if (const FObjectProperty* AsObjectProperty = CastField<FObjectProperty>(Property))
+	{
+		if (AsObjectProperty->HasAnyPropertyFlags(CPF_InstancedReference))
 		{
-			uint32 GrandChildCount;
-			ChildHandle->GetNumChildren(GrandChildCount);
-			for (uint32 I = 0; I < GrandChildCount; ++I)
+			TObjectPtr<UObject> Object = AsObjectProperty->GetObjectPropertyValue(Value);
+			UE::FPropertyBagRepository& PropertyBagRepository = UE::FPropertyBagRepository::Get();
+			if (UObject* Found = PropertyBagRepository.FindInstanceDataObject(Object))
 			{
-				const TSharedRef<IPropertyHandle> GrandChildHandle = ChildHandle->GetChildHandle(I).ToSharedRef();
-				ChildrenBuilder.AddCustomBuilder(MakeShared<FInstanceDataObjectFixupDetailNodeBuilder>(Panel.ToSharedRef(), GrandChildHandle));
+				Object = Found;
 			}
-			return;
+			GetRedirectOptions(Object->GetClass(), Object, Path, OutPaths);
 		}
 	}
-	
-	for (uint32 I = 0; I < ChildCount; ++I)
+	else if (const FArrayProperty* AsArrayProperty = CastField<FArrayProperty>(Property))
 	{
-		const TSharedRef<IPropertyHandle> ChildHandle = PropertyHandle->GetChildHandle(I).ToSharedRef();
-		ChildrenBuilder.AddCustomBuilder(MakeShared<FInstanceDataObjectFixupDetailNodeBuilder>(Panel.ToSharedRef(), ChildHandle));
+		FScriptArrayHelper Array(AsArrayProperty, Value);
+		for (int32 ArrayIndex = 0; ArrayIndex < Array.Num(); ++ArrayIndex)
+		{
+			const FPropertyPath SubPath = Path.ExtendPath(FPropertyInfo(AsArrayProperty->Inner, ArrayIndex)).Get();
+			GetRedirectOptions(AsArrayProperty->Inner, Array.GetElementPtr(ArrayIndex), SubPath, OutPaths);
+		}
+	}
+	else if (const FSetProperty* AsSetProperty = CastField<FSetProperty>(Property))
+	{
+		FScriptSetHelper Set(AsSetProperty, Value);
+		for (FScriptSetHelper::FIterator Iterator = Set.CreateIterator(); Iterator; ++Iterator)
+		{
+			const FPropertyPath SubPath = Path.ExtendPath(FPropertyInfo(AsSetProperty->ElementProp, Iterator.GetLogicalIndex())).Get();
+			GetRedirectOptions(AsSetProperty->ElementProp, Set.GetElementPtr(Iterator), SubPath, OutPaths);
+		}
+	}
+	else if (const FMapProperty* AsMapProperty = CastField<FMapProperty>(Property))
+	{
+		FScriptMapHelper Map(AsMapProperty, Value);
+		for (FScriptMapHelper::FIterator Iterator = Map.CreateIterator(); Iterator; ++Iterator)
+		{
+			const FPropertyPath SubPath = Path.ExtendPath(FPropertyInfo(AsMapProperty->ValueProp, Iterator.GetLogicalIndex())).Get();
+			GetRedirectOptions(AsMapProperty->ValueProp, Map.GetValuePtr(Iterator), SubPath, OutPaths);
+		}
 	}
 }
 
-FName FInstanceDataObjectFixupDetailNodeBuilder::GetName() const
-{
-	return PropertyHandle->GetProperty()->GetFName();
-}
-
-TSharedPtr<IPropertyHandle> FInstanceDataObjectFixupDetailNodeBuilder::GetPropertyHandle() const
-{
-	return PropertyHandle;
-}
-
-int32 FInstanceDataObjectFixupDetailNodeBuilder::GetNameWidgetIndex() const
+int32 FInstanceDataObjectNameWidgetOverride::GetNameWidgetIndex(FPropertyPath Path) const
 {
 	if (const TSharedPtr<FInstanceDataObjectFixupPanel> Panel = DiffPanel.Pin())
 	{
 		if (Panel->HasViewFlag(FInstanceDataObjectFixupPanel::EViewFlags::AllowRemapLooseProperties))
 		{
-			if (Panel->RevertInfo.Contains(*PropertyHandle->CreateFPropertyPath()))
+			if (Panel->RevertInfo.Contains(Path))
 			{
 				return DisplayRedirectMenu;
 			}
-			if (const FProperty* Property = PropertyHandle->GetProperty())
+			if (const FProperty* Property = Path.GetLeafMostProperty().Property.Get())
 			{
 				if (Property->GetBoolMetaData(TEXT("isLoose")))
 				{
@@ -364,7 +255,7 @@ int32 FInstanceDataObjectFixupDetailNodeBuilder::GetNameWidgetIndex() const
 	return DisplayRegularName;
 }
 
-TSharedRef<SWidget> FInstanceDataObjectFixupDetailNodeBuilder::GeneratePropertyRedirectMenu() const
+TSharedRef<SWidget> FInstanceDataObjectNameWidgetOverride::GeneratePropertyRedirectMenu(FPropertyPath Path) const
 {
 	FMenuBuilder MenuBuilder(true, nullptr);
 	
@@ -373,7 +264,6 @@ TSharedRef<SWidget> FInstanceDataObjectFixupDetailNodeBuilder::GeneratePropertyR
 	{
 		return MenuBuilder.MakeWidget();
 	}
-	const FPropertyPath Path = *PropertyHandle->CreateFPropertyPath();
 	const FPropertyPath& OriginalPath = Panel->GetOriginalPath(Path);
 
 	MenuBuilder.BeginSection(NAME_None, LOCTEXT("ResetRedirect", "Reset"));
@@ -395,11 +285,11 @@ TSharedRef<SWidget> FInstanceDataObjectFixupDetailNodeBuilder::GeneratePropertyR
 	{
 		for (const FPropertyPath& Option : RedirectOptions)
 		{
-			FName ThisPropName = PropertyHandle->GetProperty()->GetFName();
-			FName OptionPropName = Option.GetLeafMostProperty().Property->GetFName();
-			if (ThisPropName == OptionPropName)
+			FProperty* ThisProperty = Path.GetLeafMostProperty().Property.Get();
+			FProperty* OptionProperty = Option.GetLeafMostProperty().Property.Get();
+			if (ThisProperty->GetFName() == OptionProperty->GetFName())
 			{
-				if (Option.GetLeafMostProperty().Property->SameType(PropertyHandle->GetProperty()))
+				if (OptionProperty->SameType(ThisProperty))
 				{
 					FText DisplayName = FText::FromString(Option.ToString());
 					FText Tooltip = FText::Format(LOCTEXT("MovePropertyTooltip", "Move property to '{0}'"), DisplayName);
@@ -417,13 +307,13 @@ TSharedRef<SWidget> FInstanceDataObjectFixupDetailNodeBuilder::GeneratePropertyR
 	{
 		for (const FPropertyPath& Option : RedirectOptions)
 		{
-			FName ThisPropName = PropertyHandle->GetProperty()->GetFName();
-			FName OptionPropName = Option.GetLeafMostProperty().Property->GetFName();
-			if (ThisPropName == OptionPropName)
+			FProperty* ThisProperty = Path.GetLeafMostProperty().Property.Get();
+			FProperty* OptionProperty = Option.GetLeafMostProperty().Property.Get();
+			if (ThisProperty->GetFName() == OptionProperty->GetFName())
 			{
 				continue; // handled in the "move" category
 			}
-			if (Option.GetLeafMostProperty().Property->SameType(PropertyHandle->GetProperty()))
+			if (OptionProperty->SameType(ThisProperty))
 			{
 				FText DisplayName = FText::FromString(Option.ToString());
 				FText Tooltip = FText::Format(LOCTEXT("RenamePropertyTooltip", "Rename property to '{0}'"), DisplayName);
@@ -459,51 +349,70 @@ TSharedRef<SWidget> FInstanceDataObjectFixupDetailNodeBuilder::GeneratePropertyR
 	return MenuBuilder.MakeWidget();
 }
 
-EVisibility FInstanceDataObjectFixupDetailNodeBuilder::DeletionSymbolVisibility() const
+EVisibility FInstanceDataObjectNameWidgetOverride::DeletionSymbolVisibility(FPropertyPath Path) const
 {
 	if (const TSharedPtr<FInstanceDataObjectFixupPanel> Panel = DiffPanel.Pin())
 	{
-		const FPropertyPath Path = *PropertyHandle->CreateFPropertyPath();
 		return Panel->MarkedForDelete.Contains(Path) ? EVisibility::Visible :  EVisibility::Collapsed;
 	}
 	return EVisibility::Collapsed;
 }
 
-EVisibility FInstanceDataObjectFixupDetailNodeBuilder::ValueContentVisibility() const
+EVisibility FInstanceDataObjectNameWidgetOverride::ValueContentVisibility(FPropertyPath Path) const
 {
 	if (const TSharedPtr<FInstanceDataObjectFixupPanel> Panel = DiffPanel.Pin())
 	{
-		const FPropertyPath Path = *PropertyHandle->CreateFPropertyPath();
 		return Panel->MarkedForDelete.Contains(Path) ? EVisibility::Collapsed :  EVisibility::Visible;
 	}
 	return EVisibility::Visible;
 }
 
-bool FInstanceDataObjectFixupDetailNodeBuilder::IsHidden() const
+/////////////////////////////////////////////////////////////////
+// FHideLoosePropertiesCustomization
+/////////////////////////////////////////////////////////////////
+
+FHideLoosePropertiesCustomization::~FHideLoosePropertiesCustomization()
 {
-	const TSharedPtr<FInstanceDataObjectFixupPanel> Panel = DiffPanel.Pin();
-	if (!Panel)
+}
+
+void FHideLoosePropertiesCustomization::CustomizeDetails(IDetailLayoutBuilder& DetailBuilder)
+{
+	TArray<TWeakObjectPtr<UObject>> Objects;
+	DetailBuilder.GetObjectsBeingCustomized(Objects);
+	if (Objects.IsEmpty())
 	{
-		return true;
+		return;
 	}
-
-	if (const FProperty* Property = PropertyHandle->GetProperty())
+	TArray<FName> CategoryNames;
+	DetailBuilder.GetCategoryNames(CategoryNames);
+	for (const FName CategoryName : CategoryNames)
 	{
-		if (Panel->HasViewFlag(FInstanceDataObjectFixupPanel::EViewFlags::HideLooseProperties) &&
-			Property->GetBoolMetaData(TEXT("isLoose")))
+		IDetailCategoryBuilder& Category = DetailBuilder.EditCategory(CategoryName);
+		TArray<TSharedRef<IPropertyHandle>> Handles;
+		Category.GetDefaultProperties(Handles, true, true);
+		for (const TSharedRef<IPropertyHandle>& Handle : Handles)
 		{
-			return true;
+			CustomizeHandle(Handle, DetailBuilder);
 		}
+	}
+}
 
-		if (Panel->HasViewFlag(FInstanceDataObjectFixupPanel::EViewFlags::IncludeOnlySetBySerialization))
+void FHideLoosePropertiesCustomization::CustomizeHandle(const TSharedRef<IPropertyHandle>& Handle, IDetailLayoutBuilder& DetailBuilder)
+{
+	if (Handle->GetBoolMetaData(TEXT("isLoose")))
+	{
+		DetailBuilder.HideProperty(Handle);
+	}
+	else
+	{
+		// recurse into children
+		uint32 NumChildren = 0;
+		Handle->GetNumChildren(NumChildren);
+		for (uint32 ChildIndex = 0; ChildIndex < NumChildren; ++ChildIndex)
 		{
-			if (!Panel->IsInRedirectedPropertyTree(*PropertyHandle->CreateFPropertyPath()))
-			{
-				return true;
-			}
+			CustomizeHandle(Handle->GetChildHandle(ChildIndex).ToSharedRef(), DetailBuilder);
 		}
 	}
-	return false;
 }
 
 #undef LOCTEXT_NAMESPACE
diff --git a/Engine/Plugins/Experimental/InstanceDataObjectFixupTool/Source/Private/InstanceDataObjectFixupDetailCustomization.h b/Engine/Plugins/Experimental/InstanceDataObjectFixupTool/Source/Private/InstanceDataObjectFixupDetailCustomization.h
--- a/Engine/Plugins/Experimental/InstanceDataObjectFixupTool/Source/Private/InstanceDataObjectFixupDetailCustomization.h
+++ b/Engine/Plugins/Experimental/InstanceDataObjectFixupTool/Source/Private/InstanceDataObjectFixupDetailCustomization.h
@@ -2,6 +2,8 @@
 
 #pragma once
 
+#include <DetailsNameWidgetOverrideCustomization.h>
+
 #include "CoreMinimal.h"
 #include "IDetailCustomization.h"
 #include "IDetailCustomNodeBuilder.h"
@@ -23,44 +25,29 @@ public:
 	virtual void CustomizeDetails(IDetailLayoutBuilder& DetailBuilder) override;
 
 private:
+	bool IsHidden(const TSharedPtr<IPropertyHandle>& PropertyHandle) const;
+	void CustomizeHandle(const TSharedRef<IPropertyHandle>& Handle, IDetailLayoutBuilder& DetailBuilder);
 	TWeakPtr<FInstanceDataObjectFixupPanel> DiffPanel;
 };
 
-class INSTANCEDATAOBJECTFIXUPTOOL_API FHideLoosePropertiesCustomization : public IDetailCustomization, public TSharedFromThis<FInstanceDataObjectFixupDetailCustomization>
+class INSTANCEDATAOBJECTFIXUPTOOL_API FInstanceDataObjectNameWidgetOverride : public FDetailsNameWidgetOverrideCustomization
 {
 public:
-	FHideLoosePropertiesCustomization() = default;
-	virtual ~FHideLoosePropertiesCustomization() override;
-	virtual void CustomizeDetails(IDetailLayoutBuilder& DetailBuilder) override;
-
-	static void CustomizeHandle(const TSharedRef<IPropertyHandle>& Handle, IDetailLayoutBuilder& DetailBuilder);
-
-private:
-	TWeakPtr<FInstanceDataObjectFixupPanel> DiffPanel;
-};
-
-class INSTANCEDATAOBJECTFIXUPTOOL_API FInstanceDataObjectFixupDetailNodeBuilder : public IDetailCustomNodeBuilder, public TSharedFromThis<FInstanceDataObjectFixupDetailNodeBuilder>
-{
-public:
-	FInstanceDataObjectFixupDetailNodeBuilder(const TSharedRef<FInstanceDataObjectFixupPanel>& DiffPanel, const TSharedRef<IPropertyHandle>& PropertyHandle);
-	virtual void GenerateHeaderRowContent(FDetailWidgetRow& NodeRow) override;
-	virtual void GenerateChildContent(IDetailChildrenBuilder& ChildrenBuilder) override;
-	virtual FName GetName() const override;
-	virtual TSharedPtr<IPropertyHandle> GetPropertyHandle() const override;
-
+	FInstanceDataObjectNameWidgetOverride(const TSharedRef<FInstanceDataObjectFixupPanel>& DiffPanel);
+	virtual ~FInstanceDataObjectNameWidgetOverride() override = default;
+	virtual TSharedRef<SWidget> CustomizeName(TSharedRef<SWidget> InnerNameContent, FPropertyPath& Path) override;
 private:
-	int32 GetNameWidgetIndex() const;
-	TSharedRef<SWidget> GeneratePropertyRedirectMenu() const;
-	EVisibility DeletionSymbolVisibility() const;
-	EVisibility ValueContentVisibility() const;
-	bool IsHidden() const;
+	
+	int32 GetNameWidgetIndex(FPropertyPath Path) const;
+	TSharedRef<SWidget> GeneratePropertyRedirectMenu(FPropertyPath Path) const;
+	EVisibility DeletionSymbolVisibility(FPropertyPath Path) const;
+	EVisibility ValueContentVisibility(FPropertyPath Path) const;
 	
 	TSet<FPropertyPath> GetRedirectOptions(const UStruct* Struct, void* Value) const;
 	void GetRedirectOptions(const UStruct* Struct, void* Value, const FPropertyPath& Path, TSet<FPropertyPath>& OutPaths) const;
 	void GetRedirectOptions(const FProperty* Property, void* Value, const FPropertyPath& Path, TSet<FPropertyPath>& OutPaths) const;
 	
 	TWeakPtr<FInstanceDataObjectFixupPanel> DiffPanel;
-	TSharedRef<IPropertyHandle> PropertyHandle;
 
 	enum ENameWidgetIndex : uint8
 	{
@@ -68,3 +55,16 @@ private:
 		DisplayRedirectMenu = 1
 	};
 };
+
+class INSTANCEDATAOBJECTFIXUPTOOL_API FHideLoosePropertiesCustomization : public IDetailCustomization, public TSharedFromThis<FInstanceDataObjectFixupDetailCustomization>
+{
+public:
+	FHideLoosePropertiesCustomization() = default;
+	virtual ~FHideLoosePropertiesCustomization() override;
+	virtual void CustomizeDetails(IDetailLayoutBuilder& DetailBuilder) override;
+
+	static void CustomizeHandle(const TSharedRef<IPropertyHandle>& Handle, IDetailLayoutBuilder& DetailBuilder);
+
+private:
+	TWeakPtr<FInstanceDataObjectFixupPanel> DiffPanel;
+};
diff --git a/Engine/Plugins/Experimental/InstanceDataObjectFixupTool/Source/Private/InstanceDataObjectFixupPanel.cpp b/Engine/Plugins/Experimental/InstanceDataObjectFixupTool/Source/Private/InstanceDataObjectFixupPanel.cpp
--- a/Engine/Plugins/Experimental/InstanceDataObjectFixupTool/Source/Private/InstanceDataObjectFixupPanel.cpp
+++ b/Engine/Plugins/Experimental/InstanceDataObjectFixupTool/Source/Private/InstanceDataObjectFixupPanel.cpp
@@ -295,6 +295,7 @@ TSharedPtr<IDetailsView>& FInstanceDataObjectFixupPanel::GenerateDetailsView(boo
 	DetailsViewArgs.bHideSelectionTip = true;
 	DetailsViewArgs.ExternalScrollbar = SAssignNew(LinkableScrollBar, SLinkableScrollBar);
 	DetailsViewArgs.ScrollbarAlignment = bScrollbarOnLeft ? HAlign_Left : HAlign_Right;
+	DetailsViewArgs.DetailsNameWidgetOverrideCustomization = MakeShared<FInstanceDataObjectNameWidgetOverride>(SharedThis(this));
 	
 	FPropertyEditorModule& PropertyEditorModule = FModuleManager::GetModuleChecked<FPropertyEditorModule>("PropertyEditor");
 	DetailsView = PropertyEditorModule.CreateDetailView(DetailsViewArgs);
@@ -598,6 +599,7 @@ void FInstanceDataObjectFixupPanel::RedirectProperty(const FPropertyPath& From,
 		if (SourceProperty->SameType(DestinationProperty))
 		{
 			SourceProperty->CopyCompleteValue(Destination, Source);
+			FOverridableManager::Get().GetOverriddenProperties(*Instance)->SetOverriddenPropertyOperation(EOverriddenPropertyOperation::Modified, nullptr, DestinationProperty);
 		}
 		else
 		{
diff --git a/Engine/Plugins/Experimental/InstanceDataObjectFixupTool/Source/Private/InstanceDataObjectFixupPanel.h b/Engine/Plugins/Experimental/InstanceDataObjectFixupTool/Source/Private/InstanceDataObjectFixupPanel.h
--- a/Engine/Plugins/Experimental/InstanceDataObjectFixupTool/Source/Private/InstanceDataObjectFixupPanel.h
+++ b/Engine/Plugins/Experimental/InstanceDataObjectFixupTool/Source/Private/InstanceDataObjectFixupPanel.h
@@ -85,7 +85,7 @@ public:
 
 private:
 	friend class FInstanceDataObjectFixupSpecification; // for access to Redirects
-	friend class FInstanceDataObjectFixupDetailNodeBuilder;
+	friend class FInstanceDataObjectNameWidgetOverride;
 	friend class UInstanceDataObjectFixupUndoHandler;
 	
 	void RedirectProperty(const FPropertyPath& From, const FPropertyPath& To);
