Marked InstanceDataObjectFixupTool as experimental in the uplugin file
#pf skip

diff --git a/Engine/Plugins/Experimental/InstanceDataObjectFixupTool/InstanceDataObjectFixupTool.uplugin b/Engine/Plugins/Experimental/InstanceDataObjectFixupTool/InstanceDataObjectFixupTool.uplugin
--- a/Engine/Plugins/Experimental/InstanceDataObjectFixupTool/InstanceDataObjectFixupTool.uplugin
+++ b/Engine/Plugins/Experimental/InstanceDataObjectFixupTool/InstanceDataObjectFixupTool.uplugin
@@ -12,7 +12,8 @@
 	"SupportURL" : "",
 	"EnabledByDefault" : false,
 	"CanContainContent" : false,
-	"IsBetaVersion" : true,
+	"IsBetaVersion" : false,
+	"IsExperimentalVersion" : true,
 	"Installed" : false,
 	"Plugins": 
 	[
