[Feature] Fixup tool type conversions

diff --git a/Engine/Plugins/Experimental/InstanceDataObjectFixupTool/InstanceDataObjectFixupTool.uplugin b/Engine/Plugins/Experimental/InstanceDataObjectFixupTool/InstanceDataObjectFixupTool.uplugin
--- a/Engine/Plugins/Experimental/InstanceDataObjectFixupTool/InstanceDataObjectFixupTool.uplugin
+++ b/Engine/Plugins/Experimental/InstanceDataObjectFixupTool/InstanceDataObjectFixupTool.uplugin
@@ -20,6 +20,10 @@
 		{
 			"Name": "StructUtils",
 			"Enabled": true
+		},
+		{
+			"Name": "Solaris",
+			"Enabled": true
 		}
 	],
 	"Modules" :
diff --git a/Engine/Plugins/Experimental/InstanceDataObjectFixupTool/Source/InstanceDataObjectFixupTool.Build.cs b/Engine/Plugins/Experimental/InstanceDataObjectFixupTool/Source/InstanceDataObjectFixupTool.Build.cs
--- a/Engine/Plugins/Experimental/InstanceDataObjectFixupTool/Source/InstanceDataObjectFixupTool.Build.cs
+++ b/Engine/Plugins/Experimental/InstanceDataObjectFixupTool/Source/InstanceDataObjectFixupTool.Build.cs
@@ -21,7 +21,8 @@ namespace UnrealBuildTool.Rules
 					"Slate",
 					"PropertyEditor", 
 					"WorkspaceMenuStructure", 
-					"UnrealEd",
+					"UnrealEd", 
+					"VerseNative",
 				}
 			);
 		}
diff --git a/Engine/Plugins/Experimental/InstanceDataObjectFixupTool/Source/Private/InstanceDataObjectFixupDetailCustomization.cpp b/Engine/Plugins/Experimental/InstanceDataObjectFixupTool/Source/Private/InstanceDataObjectFixupDetailCustomization.cpp
--- a/Engine/Plugins/Experimental/InstanceDataObjectFixupTool/Source/Private/InstanceDataObjectFixupDetailCustomization.cpp
+++ b/Engine/Plugins/Experimental/InstanceDataObjectFixupTool/Source/Private/InstanceDataObjectFixupDetailCustomization.cpp
@@ -157,7 +157,15 @@ TSet<FPropertyPath> FInstanceDataObjectNameWidgetOverride::GetRedirectOptions(co
 void FInstanceDataObjectNameWidgetOverride::GetRedirectOptions(const UStruct* Struct, void* Value, const FPropertyPath& Path, TSet<FPropertyPath>& OutPaths) const
 {
 	for (FProperty* SubProperty : TFieldRange<FProperty>(Struct))
-    {
+	{
+		if (SubProperty->HasAnyPropertyFlags(CPF_Transient))
+		{
+			continue;
+		}
+		if (!SubProperty->HasAnyPropertyFlags(CPF_Edit | CPF_EditConst))
+		{
+			continue;
+		}
 		if (SubProperty->ArrayDim == 1)
 		{
 			const FPropertyPath SubPath = Path.ExtendPath(FPropertyInfo(SubProperty)).Get();
@@ -277,6 +285,20 @@ TSharedRef<SWidget> FInstanceDataObjectNameWidgetOverride::GeneratePropertyRedir
 						, EUserInterfaceActionType::RadioButton);
 	}
 	MenuBuilder.EndSection();
+	
+	MenuBuilder.BeginSection(NAME_None, LOCTEXT("Delete", "Delete"));
+	{
+		if (!Panel->MarkedForDelete.Contains(OriginalPath)) // check that it wasn't already marked as deleted
+		{
+			FText DisplayName = LOCTEXT("MarkForDeletion", "Mark For Deletion");
+			FText Tooltip = LOCTEXT("MarkForDeletionTooltip", "Mark this property for deletion");
+			MenuBuilder.AddMenuEntry(DisplayName, Tooltip, FSlateIcon()
+						, FUIAction(FExecuteAction::CreateSP(Panel.Get(), &FInstanceDataObjectFixupPanel::OnMarkForDelete, Path))
+						, NAME_None
+						, EUserInterfaceActionType::RadioButton);
+		}
+	}
+	MenuBuilder.EndSection();
 
 	UObject* FirstInstanceDataObject = Panel->Instances[0];
 	TSet<FPropertyPath> RedirectOptions = GetRedirectOptions(FirstInstanceDataObject->GetClass(), FirstInstanceDataObject);
@@ -326,22 +348,71 @@ TSharedRef<SWidget> FInstanceDataObjectNameWidgetOverride::GeneratePropertyRedir
 	}
 	MenuBuilder.EndSection();
 
-	MenuBuilder.BeginSection(NAME_None, LOCTEXT("ChangeToDifferingTypeConversion", "Convert Type (potential data loss)"));
+	MenuBuilder.BeginSection(NAME_None, LOCTEXT("ChangeToDifferingTypeConversion", "Convert Type"));
 	{
-		// TODO: Add type conversions?
+		for (const FPropertyPath& Option : RedirectOptions)
+		{
+			FProperty* ThisProperty = Path.GetLeafMostProperty().Property.Get();
+			FProperty* OptionProperty = Option.GetLeafMostProperty().Property.Get();
+			if (OptionProperty->SameType(ThisProperty))
+			{
+				continue; // same type handled above
+			}
+			if (ThisProperty->GetFName() != OptionProperty->GetFName())
+			{
+				continue; // renames to handled below
+			}
+			if (FInstanceDataObjectFixupPanel::FTypeConverter Converter = Panel->CreateTypeConverter(Path, Option))
+			{
+				FText DisplayName = FText::FromString(Option.ToString());
+				FText TypeName = FText::FromName(Option.GetLeafMostProperty().Property->GetID());
+				FText Warning = Converter.GetWarning();
+				FText Tooltip = FText::Format(LOCTEXT("ConvertTypeTooltip", "Change type to {0}"), TypeName);
+				if (!Warning.IsEmpty())
+				{
+					DisplayName = FText::Format(LOCTEXT("ConvertTypeDisplayNameWithWarning", "⚠{0}"), DisplayName);
+					Tooltip = FText::Format(LOCTEXT("ConvertTypeTooltipWithWarning", "⚠{0}\n Warning: {1}"), Tooltip, Warning);
+				}
+				MenuBuilder.AddMenuEntry(DisplayName, Tooltip, FSlateIcon()
+				, FUIAction(FExecuteAction::CreateSP(Panel.Get(), &FInstanceDataObjectFixupPanel::OnRedirectProperty, Path, Option, Converter))
+				, NAME_None
+				, EUserInterfaceActionType::RadioButton);
+			}
+		}
 	}
 	MenuBuilder.EndSection();
-	
-	MenuBuilder.BeginSection(NAME_None, LOCTEXT("Delete", "Delete"));
+
+	MenuBuilder.BeginSection(NAME_None, LOCTEXT("RenameToDifferingTypeConversion", "Convert Type and Rename"));
 	{
-		if (!Panel->MarkedForDelete.Contains(OriginalPath)) // check that it wasn't already marked as deleted
+		for (const FPropertyPath& Option : RedirectOptions)
 		{
-			FText DisplayName = LOCTEXT("MarkForDeletion", "Mark For Deletion");
-			FText Tooltip = LOCTEXT("MarkForDeletionTooltip", "Mark this property for deletion");
-			MenuBuilder.AddMenuEntry(DisplayName, Tooltip, FSlateIcon()
-						, FUIAction(FExecuteAction::CreateSP(Panel.Get(), &FInstanceDataObjectFixupPanel::OnMarkForDelete, Path))
-						, NAME_None
-						, EUserInterfaceActionType::RadioButton);
+			FProperty* ThisProperty = Path.GetLeafMostProperty().Property.Get();
+			FProperty* OptionProperty = Option.GetLeafMostProperty().Property.Get();
+			if (OptionProperty->SameType(ThisProperty))
+			{
+				continue; // same type handled above
+			}
+			if (ThisProperty->GetFName() == OptionProperty->GetFName())
+			{
+				continue; // renames to handled below
+			}
+			if (FInstanceDataObjectFixupPanel::FTypeConverter Converter = Panel->CreateTypeConverter(Path, Option))
+			{
+				FText DisplayName = FText::FromString(Option.ToString());
+				FText PropDisplayName = Option.GetLeafMostProperty().Property->GetDisplayNameText();
+				FText TypeName = FText::FromName(Option.GetLeafMostProperty().Property->GetID());
+				FText Warning = Converter.GetWarning();
+				FText Tooltip = FText::Format(LOCTEXT("ConvertTypeAndRenameTooltip", "Change type to {0} and rename to {1}"), TypeName, PropDisplayName);
+				if (!Warning.IsEmpty())
+				{
+					DisplayName = FText::Format(LOCTEXT("ConvertTypeAndRenameDisplayNameWithWarning", "⚠{0}"), DisplayName);
+					Tooltip = FText::Format(LOCTEXT("ConvertTypeAndRenameTooltipWithWarning", "⚠{0}\n Warning: {1}"), Tooltip, Warning);
+				}
+				MenuBuilder.AddMenuEntry(DisplayName, Tooltip, FSlateIcon()
+				, FUIAction(FExecuteAction::CreateSP(Panel.Get(), &FInstanceDataObjectFixupPanel::OnRedirectProperty, Path, Option, Converter))
+				, NAME_None
+				, EUserInterfaceActionType::RadioButton);
+			}
 		}
 	}
 	MenuBuilder.EndSection();
diff --git a/Engine/Plugins/Experimental/InstanceDataObjectFixupTool/Source/Private/InstanceDataObjectFixupPanel.cpp b/Engine/Plugins/Experimental/InstanceDataObjectFixupTool/Source/Private/InstanceDataObjectFixupPanel.cpp
--- a/Engine/Plugins/Experimental/InstanceDataObjectFixupTool/Source/Private/InstanceDataObjectFixupPanel.cpp
+++ b/Engine/Plugins/Experimental/InstanceDataObjectFixupTool/Source/Private/InstanceDataObjectFixupPanel.cpp
@@ -9,10 +9,12 @@
 #include "InstanceDataObjectFixupDetailCustomization.h"
 #include "Modules/ModuleManager.h"
 #include "Editor.h"
+#include "PropertyVerseString.h"
 #include "UObject/PropertyBagRepository.h"
 
 #include "UObject/OverriddenPropertySet.h"
 #include "UObject/OverridableManager.h"
+#include "UObject/TextProperty.h"
 
 #define LOCTEXT_NAMESPACE "InstanceDataObjectFixupPanel"
 
@@ -488,7 +490,161 @@ static FPropertyChangedEvent ConstructChangeEventForRedirect(const FPropertyPath
 	return OutEvent;
 }
 
-void FInstanceDataObjectFixupPanel::RedirectProperty(const FPropertyPath& From, const FPropertyPath& To)
+void FInstanceDataObjectFixupPanel::FTypeConverter::Push(FProperty* SourceProperty, const void* SourceData, FProperty* DestinationProperty, void* DestinationData)
+{
+	InstanceInfo.Push({SourceProperty, SourceData, DestinationProperty, DestinationData});
+	// check if warning was made more severe by this data
+	Warning = FMath::Max(Warning, GenerateWarning(SourceProperty, SourceData, DestinationProperty));
+}
+
+FInstanceDataObjectFixupPanel::FTypeConverter::operator bool() const
+{
+	return Warning != EWarning::InvalidConversion;
+}
+
+void FInstanceDataObjectFixupPanel::FTypeConverter::operator()() const
+{
+	check(Warning != EWarning::InvalidConversion);
+	for (const FInstanceInfo& Info : InstanceInfo)
+	{
+		TryConvert(Info.SourceProperty, Info.SourceData, Info.DestinationProperty, Info.DestinationData);
+	}
+}
+
+FText FInstanceDataObjectFixupPanel::FTypeConverter::GetWarning() const
+{
+	switch(Warning)
+	{
+	case EWarning::NarrowingConversion:
+		return LOCTEXT("NarrowingConversion", "This type conversion is a narrowing conversion. Likely data loss!");
+	case EWarning::NonInvertibleConversion:
+		return LOCTEXT("NonInvertibleConversion", "This type conversion is not an invertable operation. Likely data loss!");
+	case EWarning::InvalidConversion:
+		return LOCTEXT("InvalidConversion", "Invalid Conversion");
+	default:
+		return FText::GetEmpty();
+	}
+}
+
+bool FInstanceDataObjectFixupPanel::FTypeConverter::TryConvert(FProperty* SourceProperty, const void* SourceData, FProperty* DestinationProperty, void* DestinationData)
+{
+	TArray<uint8, TInlineAllocator<64>> Buffer;
+	TMemoryWriterBase<TInlineAllocator<64>> MemoryWriter(Buffer);
+	FStructuredArchiveFromArchive StructuredWriter(MemoryWriter);
+	SourceProperty->SerializeItem(StructuredWriter.GetSlot(), (uint8*)SourceData);
+	FMemoryReaderView MemoryReader(Buffer);
+	FStructuredArchiveFromArchive StructuredReader(MemoryReader);
+
+	// TODO: this breaks for static array elements.
+	void* DestinationContainer = static_cast<uint8*>(DestinationData) - DestinationProperty->GetOffset_ForInternal();
+
+	// todo: handle static arrays
+	FPropertyTag SourceTag(MemoryReader, SourceProperty, 0, (uint8*)SourceData, nullptr);
+
+	bool bResult = false;
+	switch(DestinationProperty->ConvertFromType(SourceTag, StructuredReader.GetSlot(), (uint8*)DestinationContainer, SourceProperty->GetOwnerStruct(), nullptr))
+	{
+	case EConvertFromTypeResult::UseSerializeItem:
+		if (SourceProperty->GetID() == DestinationProperty->GetID())
+		{
+			SourceTag.SerializeTaggedProperty(StructuredReader.GetSlot(), DestinationProperty, (uint8*)DestinationContainer, nullptr);
+			bResult = true;
+		}
+		break;
+	case EConvertFromTypeResult::Serialized:
+		bResult = true;
+		break;
+	case EConvertFromTypeResult::CannotConvert:
+		break;
+	case EConvertFromTypeResult::Converted:
+		bResult = true;
+		break;
+	}
+
+	if (!bResult)
+	{
+		bool bTryTextSerialize = false;
+
+		const auto IsStringType = [](const FProperty* Property)
+		{
+			return Property->IsA<FStrProperty>() || Property->IsA<FTextProperty>() || Property->IsA<FNameProperty>() || Property->IsA<FVerseStringProperty>();
+		};
+		
+		if (IsStringType(SourceProperty) || IsStringType(DestinationProperty))
+		{
+			// if either property is a string, text, or name, use text serialization
+			bTryTextSerialize = true;
+		}
+		else if (FStructProperty* SourceAsStructProperty = CastField<FStructProperty>(SourceProperty))
+		{
+			if (FStructProperty* DestinationAsStructProperty = CastField<FStructProperty>(DestinationProperty))
+			{
+				if (!SourceAsStructProperty->Struct->UseNativeSerialization() && !DestinationAsStructProperty->Struct->UseNativeSerialization())
+				{
+					// attempt to text serialize structs since ConvertFromType doesn't support them usually
+					bTryTextSerialize = true;
+				}
+			}
+		}
+
+		// use ExportText_Direct and ImportText_Direct
+		if (bTryTextSerialize)
+		{
+			FString StrBuffer;
+			SourceProperty->ExportText_Direct(StrBuffer, SourceData, nullptr, nullptr, PPF_None);
+			FStringOutputDevice ErrorOutput;
+			DestinationProperty->ImportText_Direct(*StrBuffer, DestinationData, nullptr, PPF_None, &ErrorOutput);
+			bResult = ErrorOutput.IsEmpty();
+		}
+	}
+
+	
+	return bResult;
+}
+
+FInstanceDataObjectFixupPanel::FTypeConverter::EWarning FInstanceDataObjectFixupPanel::FTypeConverter::GenerateWarning(FProperty* SourceProperty, const void* SourceData, FProperty* DestinationProperty)
+{
+	// convert from source to destination in a temp buffer to see if it's possible
+	TArray<uint8, TInlineAllocator<64>> SourceToDest;
+	SourceToDest.SetNumUninitialized(DestinationProperty->ElementSize);
+	DestinationProperty->InitializeValue(SourceToDest.GetData());
+	if (!TryConvert(SourceProperty, SourceData, DestinationProperty, SourceToDest.GetData()))
+	{
+		return EWarning::InvalidConversion;
+	}
+
+	// convert from destination to source in a temp buffer to see if it's possible
+	TArray<uint8, TInlineAllocator<64>> DestToSource;
+	DestToSource.SetNumUninitialized(SourceProperty->ElementSize);
+	SourceProperty->InitializeValue(DestToSource.GetData());
+	if (!TryConvert(DestinationProperty, SourceToDest.GetData(), SourceProperty, DestToSource.GetData()))
+	{
+		return EWarning::NonInvertibleConversion;
+	}
+
+	// check that the round trip result has the same value as source
+	if (!SourceProperty->Identical(SourceData, DestToSource.GetData(), PPF_None))
+	{
+		return EWarning::NarrowingConversion;
+	}
+	return EWarning::SafeConversion;
+}
+
+FInstanceDataObjectFixupPanel::FTypeConverter FInstanceDataObjectFixupPanel::CreateTypeConverter(const FPropertyPath& From, const FPropertyPath& To)
+{
+	FTypeConverter Result;
+	for (UObject* Instance : Instances)
+	{
+		FProperty* SourceProperty = From.GetLeafMostProperty().Property.Get();
+		const void* SourceData = ResolvePath(From, Instance);
+		FProperty* DestinationProperty = To.GetLeafMostProperty().Property.Get();
+		void* DestinationData = ResolvePath(To, Instance);
+		Result.Push(SourceProperty, SourceData, DestinationProperty, DestinationData);
+	}
+	return Result;
+}
+
+void FInstanceDataObjectFixupPanel::RedirectPropertyHelper(const FPropertyPath& From, const FPropertyPath& To, TOptional<FRevertInfo>& FromRevertInfo, FRevertInfo*& ToRevertInfo)
 {
 	UInstanceDataObjectFixupUndoHandler* Snapshot = NewObject<UInstanceDataObjectFixupUndoHandler>();
 	Snapshot->Init(SharedThis(this));
@@ -497,8 +653,6 @@ void FInstanceDataObjectFixupPanel::RedirectProperty(const FPropertyPath& From,
 	FProperty* SourceProperty = From.GetLeafMostProperty().Property.Get();
 	check(SourceProperty);
 	FProperty* DestinationProperty = To.IsValid() ? To.GetLeafMostProperty().Property.Get() : nullptr;
-	FRevertInfo* ToRevertInfo = nullptr;
-	TOptional<FRevertInfo> FromRevertInfo;
 	
 	if (const FRevertInfo* Info = RevertInfo.Find(From))
 	{
@@ -558,16 +712,26 @@ void FInstanceDataObjectFixupPanel::RedirectProperty(const FPropertyPath& From,
 	}
 
 	Snapshot->OnRedirect(From, To);
+}
+
+void FInstanceDataObjectFixupPanel::RedirectProperty(const FPropertyPath& From, const FPropertyPath& To)
+{
+	FProperty* SourceProperty = From.GetLeafMostProperty().Property.Get();
+	check(SourceProperty);
+	FProperty* DestinationProperty = To.IsValid() ? To.GetLeafMostProperty().Property.Get() : nullptr;
+	
+	TOptional<FRevertInfo> FromRevertInfo;
+	FRevertInfo* ToRevertInfo = nullptr;
+	RedirectPropertyHelper(From, To, FromRevertInfo, ToRevertInfo);
 	
-	if (!DestinationProperty)
+	if (!DestinationProperty) // null destination is interpreted as a deletion
 	{
 		MarkedForDelete.Add(From);
 		GEditor->EndTransaction();
-        DetailsView->ForceRefresh();
+		DetailsView->ForceRefresh();
 		return; // delete actions don't need data copied
 	}
 	
-
 	const uint8* FromRevertInfoItr = FromRevertInfo ? FromRevertInfo->OriginalValue.GetData() : nullptr;
 	for (UObject* Instance : Instances)
 	{
@@ -618,6 +782,70 @@ void FInstanceDataObjectFixupPanel::RedirectProperty(const FPropertyPath& From,
 		FOverridableManager::Get().PostOverrideProperty(*Instance, ChangeEvent, Chain);
 	}
 
+	GEditor->EndTransaction();
+	DetailsView->ForceRefresh();
+}
+
+void FInstanceDataObjectFixupPanel::RedirectProperty(const FPropertyPath& From, const FPropertyPath& To, const FTypeConverter& TypeConversion)
+{
+	const FProperty* SourceProperty = From.GetLeafMostProperty().Property.Get();
+	const FProperty* DestinationProperty = To.GetLeafMostProperty().Property.Get();
+	check(SourceProperty && DestinationProperty);
+	
+	TOptional<FRevertInfo> FromRevertInfo;
+	FRevertInfo* ToRevertInfo = nullptr;
+	RedirectPropertyHelper(From, To, FromRevertInfo, ToRevertInfo);
+	
+	const uint8* FromRevertInfoItr = FromRevertInfo ? FromRevertInfo->OriginalValue.GetData() : nullptr;
+	TArray<FPropertyChangedEvent> ChangeEvents;
+	TArray<FEditPropertyChain> Chains;
+
+	// call PreEditChange and set up undo handling
+	for (UObject* Instance : Instances)
+	{
+		void* Source = ResolvePath(From, Instance);
+		void* Destination = ResolvePath(To, Instance);
+		
+		if (!ensure(Source && Destination))
+		{
+			continue;
+		}
+		
+		// construct change event
+		Chains.Emplace();
+		TMap<FString, int32> ArrayIndices;
+		ChangeEvents.Add(ConstructChangeEventForRedirect(To, Chains.Last(), ArrayIndices));
+		FOverridableManager::Get().PreOverrideProperty(*Instance, Chains.Last());
+		Instance->PreEditChange(ChangeEvents.Last().Property);
+
+		if (ToRevertInfo)
+		{
+			// cache the destination value so it can be reverted later
+			const int32 Size = DestinationProperty->ArrayDim * DestinationProperty->ElementSize;
+			ToRevertInfo->OriginalValue.AddZeroed(Size);
+			uint8* Buffer = ToRevertInfo->OriginalValue.GetData() + (ToRevertInfo->OriginalValue.Num() - Size);
+			DestinationProperty->CopyCompleteValue(Buffer, Destination);
+		}
+	}
+
+	// applied to all instances at once
+	TypeConversion();
+
+	// call post edit change and apply undo handling
+	for (int32 I = 0; I < Instances.Num(); ++I)
+	{
+		UObject* Instance = Instances[I];
+		void* Source = ResolvePath(From, Instance);
+		if (FromRevertInfo)
+		{
+			// apply FromRevertInfo to From
+			SourceProperty->CopyCompleteValue(Source, FromRevertInfoItr);
+			FromRevertInfoItr += DestinationProperty->ArrayDim * DestinationProperty->ElementSize;
+		}
+		Instance->PostEditChangeProperty(ChangeEvents[I]);
+		FOverridableManager::Get().PostOverrideProperty(*Instance, ChangeEvents[I], Chains[I]);
+	}
+
 	GEditor->EndTransaction();
 	
 	DetailsView->ForceRefresh();
@@ -628,6 +856,11 @@ void FInstanceDataObjectFixupPanel::OnRedirectProperty(FPropertyPath From, FProp
 	RedirectProperty(From, To);
 }
 
+void FInstanceDataObjectFixupPanel::OnRedirectProperty(FPropertyPath From, FPropertyPath To, FTypeConverter TypeConversion)
+{
+	RedirectProperty(From, To, TypeConversion);
+}
+
 static void InitRedirectedPropertyTreeRec(const TSharedPtr<FRedirectedPropertyNode>& Node, FProperty* Property, void* Value, TSet<UObject*>& EnteredObjects);
 static void InitRedirectedPropertyTreeRec(const TSharedPtr<FRedirectedPropertyNode>& Node, UStruct* Struct, void* StructValue, TSet<UObject*>& EnteredObjects)
 {
diff --git a/Engine/Plugins/Experimental/InstanceDataObjectFixupTool/Source/Private/InstanceDataObjectFixupPanel.h b/Engine/Plugins/Experimental/InstanceDataObjectFixupTool/Source/Private/InstanceDataObjectFixupPanel.h
--- a/Engine/Plugins/Experimental/InstanceDataObjectFixupTool/Source/Private/InstanceDataObjectFixupPanel.h
+++ b/Engine/Plugins/Experimental/InstanceDataObjectFixupTool/Source/Private/InstanceDataObjectFixupPanel.h
@@ -83,14 +83,68 @@ public:
 	TSharedPtr<IDetailsView> DetailsView;
 	TSharedPtr<SLinkableScrollBar> LinkableScrollBar;
 
+	// functor used to generate warnings about a type conversion (or set of conversions) and invoke that conversion on demand
+	struct FTypeConverter
+	{
+		FTypeConverter() = default;
+		
+		// add a conversion that should be performed call operator is invoked.
+		void Push(FProperty* SourceProperty, const void* SourceData, FProperty* DestinationProperty, void* DestinationData);
+
+		// return whether this is a valid conversion. (only true if all conversions pushed are valid)
+		operator bool() const;
+
+		// run conversion on all pushed data
+		void operator()() const;
+
+		// return the most severe text warning for the conversions pushed.
+		FText GetWarning() const;
+	private:
+		enum class EWarning
+		{
+			// sorted by least severe to most severe.
+			SafeConversion,
+			NarrowingConversion,
+			NonInvertibleConversion,
+			InvalidConversion,
+		};
+		struct FInstanceInfo
+		{
+			FProperty* SourceProperty;
+			const void* SourceData;
+			FProperty* DestinationProperty;
+			void* DestinationData;
+		};
+		
+		static bool TryConvert(FProperty* SourceProperty, const void* SourceData, FProperty* DestinationProperty, void* DestinationData);
+		static EWarning GenerateWarning(FProperty* SourceProperty, const void* SourceData, FProperty* DestinationProperty);
+
+		// most severe conversion warning found in all the pushed data
+		EWarning Warning = EWarning::SafeConversion;
+
+		TArray<FInstanceInfo> InstanceInfo;
+	};
+	
+	FTypeConverter CreateTypeConverter(const FPropertyPath& From, const FPropertyPath& To);
+
 private:
 	friend class FInstanceDataObjectFixupSpecification; // for access to Redirects
 	friend class FInstanceDataObjectNameWidgetOverride;
 	friend class UInstanceDataObjectFixupUndoHandler;
-	
+	struct FRevertInfo
+	{
+		TArray<uint8> OriginalValue;
+		FPropertyPath OriginalPath;
+		bool bWasTransient;
+		bool bWasHidden;
+	};
+
+	void RedirectPropertyHelper(const FPropertyPath& From, const FPropertyPath& To, TOptional<FRevertInfo>& FromRevertInfo, FRevertInfo*& ToRevertInfo);
 	void RedirectProperty(const FPropertyPath& From, const FPropertyPath& To);
+	void RedirectProperty(const FPropertyPath& From, const FPropertyPath& To, const FTypeConverter& TypeConversion);
 	// pass-by-copy version for delegates. Use RedirectProperty when possible
 	void OnRedirectProperty(FPropertyPath From, FPropertyPath To);
+	void OnRedirectProperty(FPropertyPath From, FPropertyPath To, FTypeConverter TypeConversion);
 	
 	void InitRedirectedPropertyTree();
 	
@@ -100,13 +154,6 @@ private:
 	// redirected to from a floating property. The members of the tree are visible in the left panel.
 	TSharedPtr<FRedirectedPropertyNode> RedirectedPropertyTree;
 
-	struct FRevertInfo
-	{
-		TArray<uint8> OriginalValue;
-		FPropertyPath OriginalPath;
-		bool bWasTransient;
-		bool bWasHidden;
-	};
 	
 	TMap<FPropertyPath, FRevertInfo> RevertInfo;
 	TSet<FPropertyPath> MarkedForDelete;
diff --git a/Engine/Source/Runtime/CoreUObject/Public/UObject/PropertyTag.h b/Engine/Source/Runtime/CoreUObject/Public/UObject/PropertyTag.h
--- a/Engine/Source/Runtime/CoreUObject/Public/UObject/PropertyTag.h
+++ b/Engine/Source/Runtime/CoreUObject/Public/UObject/PropertyTag.h
@@ -61,7 +61,7 @@ struct FPropertyTag
 
 	// Constructors.
 	FPropertyTag();
-	FPropertyTag( FArchive& InSaveAr, FProperty* Property, int32 InIndex, uint8* Value, const uint8* Defaults );
+	UE_INTERNAL COREUOBJECT_API FPropertyTag( FArchive& InSaveAr, FProperty* Property, int32 InIndex, uint8* Value, const uint8* Defaults );
 
 	// Set optional property guid
 	void SetPropertyGuid(const FGuid& InPropertyGuid);
@@ -72,7 +72,7 @@ struct FPropertyTag
 
 	// Property serializer.
 	void SerializeTaggedProperty( FArchive& Ar, FProperty* Property, uint8* Value, const uint8* Defaults ) const;
-	void SerializeTaggedProperty(FStructuredArchive::FSlot Slot, FProperty* Property, uint8* Value, const uint8* Defaults) const;
+	UE_INTERNAL COREUOBJECT_API void SerializeTaggedProperty(FStructuredArchive::FSlot Slot, FProperty* Property, uint8* Value, const uint8* Defaults) const;
 };
 
 struct FPropertyTagScope
