[CisFix] removed NotForLicensees dependency from fixup tool plugin

diff --git a/Engine/Plugins/Experimental/InstanceDataObjectFixupTool/InstanceDataObjectFixupTool.uplugin b/Engine/Plugins/Experimental/InstanceDataObjectFixupTool/InstanceDataObjectFixupTool.uplugin
--- a/Engine/Plugins/Experimental/InstanceDataObjectFixupTool/InstanceDataObjectFixupTool.uplugin
+++ b/Engine/Plugins/Experimental/InstanceDataObjectFixupTool/InstanceDataObjectFixupTool.uplugin
@@ -20,10 +20,6 @@
 		{
 			"Name": "StructUtils",
 			"Enabled": true
-		},
-		{
-			"Name": "Solaris",
-			"Enabled": true
 		}
 	],
 	"Modules" :
diff --git a/Engine/Plugins/Experimental/InstanceDataObjectFixupTool/Source/InstanceDataObjectFixupTool.Build.cs b/Engine/Plugins/Experimental/InstanceDataObjectFixupTool/Source/InstanceDataObjectFixupTool.Build.cs
--- a/Engine/Plugins/Experimental/InstanceDataObjectFixupTool/Source/InstanceDataObjectFixupTool.Build.cs
+++ b/Engine/Plugins/Experimental/InstanceDataObjectFixupTool/Source/InstanceDataObjectFixupTool.Build.cs
@@ -21,8 +21,7 @@ namespace UnrealBuildTool.Rules
 					"Slate",
 					"PropertyEditor", 
 					"WorkspaceMenuStructure", 
-					"UnrealEd", 
-					"VerseNative",
+					"UnrealEd",
 				}
 			);
 		}
diff --git a/Engine/Plugins/Experimental/InstanceDataObjectFixupTool/Source/Private/InstanceDataObjectFixupPanel.cpp b/Engine/Plugins/Experimental/InstanceDataObjectFixupTool/Source/Private/InstanceDataObjectFixupPanel.cpp
--- a/Engine/Plugins/Experimental/InstanceDataObjectFixupTool/Source/Private/InstanceDataObjectFixupPanel.cpp
+++ b/Engine/Plugins/Experimental/InstanceDataObjectFixupTool/Source/Private/InstanceDataObjectFixupPanel.cpp
@@ -9,7 +9,6 @@
 #include "InstanceDataObjectFixupDetailCustomization.h"
 #include "Modules/ModuleManager.h"
 #include "Editor.h"
-#include "PropertyVerseString.h"
 #include "UObject/PropertyBagRepository.h"
 
 #include "UObject/OverriddenPropertySet.h"
@@ -567,7 +566,8 @@ bool FInstanceDataObjectFixupPanel::FTypeConverter::TryConvert(FProperty* Source
 
 		const auto IsStringType = [](const FProperty* Property)
 		{
-			return Property->IsA<FStrProperty>() || Property->IsA<FTextProperty>() || Property->IsA<FNameProperty>() || Property->IsA<FVerseStringProperty>();
+			static FName VerseStringName = TEXT("VerseStringProperty");
+			return Property->IsA<FStrProperty>() || Property->IsA<FTextProperty>() || Property->IsA<FNameProperty>() || Property->GetID() == VerseStringName;
 		};
 		
 		if (IsStringType(SourceProperty) || IsStringType(DestinationProperty))
