Remove LooseProperties from DetailsView
* A property tagged with "IsLoose" will be automatically filtered from the DetailsView unless the DetailsView is configured to show loose properties

diff --git a/Engine/Plugins/Experimental/InstanceDataObjectFixupTool/Source/Private/InstanceDataObjectFixupDetailCustomization.cpp b/Engine/Plugins/Experimental/InstanceDataObjectFixupTool/Source/Private/InstanceDataObjectFixupDetailCustomization.cpp
--- a/Engine/Plugins/Experimental/InstanceDataObjectFixupTool/Source/Private/InstanceDataObjectFixupDetailCustomization.cpp
+++ b/Engine/Plugins/Experimental/InstanceDataObjectFixupTool/Source/Private/InstanceDataObjectFixupDetailCustomization.cpp
@@ -14,94 +14,6 @@
 
 #define LOCTEXT_NAMESPACE "InstanceDataObjectFixupDetails"
 
-/////////////////////////////////////////////////////////////////
-// FInstanceDataObjectFixupDetailCustomization
-/////////////////////////////////////////////////////////////////
-
-FInstanceDataObjectFixupDetailCustomization::FInstanceDataObjectFixupDetailCustomization(const TSharedRef<FInstanceDataObjectFixupPanel>& InDiffPanel)
-	: DiffPanel(InDiffPanel)
-{
-	
-}
-
-FInstanceDataObjectFixupDetailCustomization::~FInstanceDataObjectFixupDetailCustomization()
-{
-}
-
-void FInstanceDataObjectFixupDetailCustomization::CustomizeDetails(IDetailLayoutBuilder& DetailBuilder)
-{
-	const TSharedPtr<FInstanceDataObjectFixupPanel> Panel = DiffPanel.Pin();
-	if (!Panel)
-	{
-		return;
-	}
-	
-	TArray<TWeakObjectPtr<UObject>> Objects;
-	DetailBuilder.GetObjectsBeingCustomized(Objects);
-	if (Objects.IsEmpty())
-	{
-		return;
-	}
-	TArray<FName> CategoryNames;
-	DetailBuilder.GetCategoryNames(CategoryNames);
-	for (const FName CategoryName : CategoryNames)
-	{
-		IDetailCategoryBuilder& Category = DetailBuilder.EditCategory(CategoryName);
-		TArray<TSharedRef<IPropertyHandle>> Handles;
-		Category.GetDefaultProperties(Handles, true, true);
-		for (const TSharedRef<IPropertyHandle>& Handle : Handles)
-		{
-			CustomizeHandle(Handle, DetailBuilder);
-		}
-	}
-}
-
-bool FInstanceDataObjectFixupDetailCustomization::IsHidden(const TSharedPtr<IPropertyHandle>& PropertyHandle) const
-{
-	const TSharedPtr<FInstanceDataObjectFixupPanel> Panel = DiffPanel.Pin();
-	if (!Panel)
-	{
-		return true;
-	}
-
-	if (const FProperty* Property = PropertyHandle->GetProperty())
-	{
-		if (Panel->HasViewFlag(FInstanceDataObjectFixupPanel::EViewFlags::HideLooseProperties) &&
-			Property->GetBoolMetaData(TEXT("isLoose")))
-		{
-			return true;
-		}
-
-		if (Panel->HasViewFlag(FInstanceDataObjectFixupPanel::EViewFlags::IncludeOnlySetBySerialization))
-		{
-			if (!Panel->IsInRedirectedPropertyTree(*PropertyHandle->CreateFPropertyPath()))
-			{
-				return true;
-			}
-		}
-	}
-	return false;
-}
-
-void FInstanceDataObjectFixupDetailCustomization::CustomizeHandle(const TSharedRef<IPropertyHandle>& Handle, IDetailLayoutBuilder& DetailBuilder)
-{
-	if (IsHidden(Handle))
-	{
-		DetailBuilder.HideProperty(Handle);
-	}
-	else
-	{
-		// recurse into children
-		uint32 NumChildren = 0;
-		Handle->GetNumChildren(NumChildren);
-		for (uint32 ChildIndex = 0; ChildIndex < NumChildren; ++ChildIndex)
-		{
-			CustomizeHandle(Handle->GetChildHandle(ChildIndex).ToSharedRef(), DetailBuilder);
-		}
-	}
-}
-
-
 /////////////////////////////////////////////////////////////////
 // FInstanceDataObjectNameWidgetOverride
 /////////////////////////////////////////////////////////////////
@@ -438,52 +350,4 @@ EVisibility FInstanceDataObjectNameWidgetOverride::ValueContentVisibility(FPrope
 	return EVisibility::Visible;
 }
 
-/////////////////////////////////////////////////////////////////
-// FHideLoosePropertiesCustomization
-/////////////////////////////////////////////////////////////////
-
-FHideLoosePropertiesCustomization::~FHideLoosePropertiesCustomization()
-{
-}
-
-void FHideLoosePropertiesCustomization::CustomizeDetails(IDetailLayoutBuilder& DetailBuilder)
-{
-	TArray<TWeakObjectPtr<UObject>> Objects;
-	DetailBuilder.GetObjectsBeingCustomized(Objects);
-	if (Objects.IsEmpty())
-	{
-		return;
-	}
-	TArray<FName> CategoryNames;
-	DetailBuilder.GetCategoryNames(CategoryNames);
-	for (const FName CategoryName : CategoryNames)
-	{
-		IDetailCategoryBuilder& Category = DetailBuilder.EditCategoryAllowNone(CategoryName);
-		TArray<TSharedRef<IPropertyHandle>> Handles;
-		Category.GetDefaultProperties(Handles, true, true);
-		for (const TSharedRef<IPropertyHandle>& Handle : Handles)
-		{
-			CustomizeHandle(Handle, DetailBuilder);
-		}
-	}
-}
-
-void FHideLoosePropertiesCustomization::CustomizeHandle(const TSharedRef<IPropertyHandle>& Handle, IDetailLayoutBuilder& DetailBuilder)
-{
-	if (Handle->GetBoolMetaData(TEXT("isLoose")))
-	{
-		DetailBuilder.HideProperty(Handle);
-	}
-	else
-	{
-		// recurse into children
-		uint32 NumChildren = 0;
-		Handle->GetNumChildren(NumChildren);
-		for (uint32 ChildIndex = 0; ChildIndex < NumChildren; ++ChildIndex)
-		{
-			CustomizeHandle(Handle->GetChildHandle(ChildIndex).ToSharedRef(), DetailBuilder);
-		}
-	}
-}
-
 #undef LOCTEXT_NAMESPACE
diff --git a/Engine/Plugins/Experimental/InstanceDataObjectFixupTool/Source/Private/InstanceDataObjectFixupDetailCustomization.h b/Engine/Plugins/Experimental/InstanceDataObjectFixupTool/Source/Private/InstanceDataObjectFixupDetailCustomization.h
--- a/Engine/Plugins/Experimental/InstanceDataObjectFixupTool/Source/Private/InstanceDataObjectFixupDetailCustomization.h
+++ b/Engine/Plugins/Experimental/InstanceDataObjectFixupTool/Source/Private/InstanceDataObjectFixupDetailCustomization.h
@@ -14,22 +14,6 @@
 class IDetailCategoryBuilder;
 class FInstanceDataObjectFixupPanel;
 
-/**
- * 
- */
-class INSTANCEDATAOBJECTFIXUPTOOL_API FInstanceDataObjectFixupDetailCustomization : public IDetailCustomization, public TSharedFromThis<FInstanceDataObjectFixupDetailCustomization>
-{
-public:
-	FInstanceDataObjectFixupDetailCustomization(const TSharedRef<FInstanceDataObjectFixupPanel>& DiffPanel);
-	virtual ~FInstanceDataObjectFixupDetailCustomization() override;
-	virtual void CustomizeDetails(IDetailLayoutBuilder& DetailBuilder) override;
-
-private:
-	bool IsHidden(const TSharedPtr<IPropertyHandle>& PropertyHandle) const;
-	void CustomizeHandle(const TSharedRef<IPropertyHandle>& Handle, IDetailLayoutBuilder& DetailBuilder);
-	TWeakPtr<FInstanceDataObjectFixupPanel> DiffPanel;
-};
-
 class INSTANCEDATAOBJECTFIXUPTOOL_API FInstanceDataObjectNameWidgetOverride : public FDetailsNameWidgetOverrideCustomization
 {
 public:
@@ -54,17 +38,4 @@ private:
 		DisplayRegularName = 0,
 		DisplayRedirectMenu = 1
 	};
-};
-
-class INSTANCEDATAOBJECTFIXUPTOOL_API FHideLoosePropertiesCustomization : public IDetailCustomization, public TSharedFromThis<FInstanceDataObjectFixupDetailCustomization>
-{
-public:
-	FHideLoosePropertiesCustomization() = default;
-	virtual ~FHideLoosePropertiesCustomization() override;
-	virtual void CustomizeDetails(IDetailLayoutBuilder& DetailBuilder) override;
-
-	static void CustomizeHandle(const TSharedRef<IPropertyHandle>& Handle, IDetailLayoutBuilder& DetailBuilder);
-
-private:
-	TWeakPtr<FInstanceDataObjectFixupPanel> DiffPanel;
-};
+};
\ No newline at end of file
diff --git a/Engine/Plugins/Experimental/InstanceDataObjectFixupTool/Source/Private/InstanceDataObjectFixupPanel.cpp b/Engine/Plugins/Experimental/InstanceDataObjectFixupTool/Source/Private/InstanceDataObjectFixupPanel.cpp
--- a/Engine/Plugins/Experimental/InstanceDataObjectFixupTool/Source/Private/InstanceDataObjectFixupPanel.cpp
+++ b/Engine/Plugins/Experimental/InstanceDataObjectFixupTool/Source/Private/InstanceDataObjectFixupPanel.cpp
@@ -17,6 +17,8 @@
 
 #define LOCTEXT_NAMESPACE "InstanceDataObjectFixupPanel"
 
+static const FName NAME_IsLooseMetadata(TEXT("IsLoose"));
+
 FRedirectedPropertyNode::FRedirectedPropertyNode(const FRedirectedPropertyNode& Other)
 	: PropertyName(Other.PropertyName)
 	, Type(Other.Type)
@@ -217,7 +219,6 @@ int32 FInstanceDataObjectFixupPanel::Find(UObject* Value) const
 static bool RemoveCustomizationsWithLooseProperties(const FFieldVariant& FieldVariant, const TSharedPtr<IDetailsView>& DetailsView)
 {
 #if WITH_EDITORONLY_DATA
-	static const FName NAME_IsLooseMetadata(TEXT("IsLoose"));
 	if (FStructProperty* AsStructProperty = FieldVariant.Get<FStructProperty>())
 	{
 		if (RemoveCustomizationsWithLooseProperties(AsStructProperty->Struct, DetailsView))
@@ -297,6 +298,7 @@ TSharedPtr<IDetailsView>& FInstanceDataObjectFixupPanel::GenerateDetailsView(boo
 	DetailsViewArgs.ExternalScrollbar = SAssignNew(LinkableScrollBar, SLinkableScrollBar);
 	DetailsViewArgs.ScrollbarAlignment = bScrollbarOnLeft ? HAlign_Left : HAlign_Right;
 	DetailsViewArgs.DetailsNameWidgetOverrideCustomization = MakeShared<FInstanceDataObjectNameWidgetOverride>(SharedThis(this));
+	DetailsViewArgs.bShowLooseProperties = !HasViewFlag(EViewFlags::HideLooseProperties);
 	
 	FPropertyEditorModule& PropertyEditorModule = FModuleManager::GetModuleChecked<FPropertyEditorModule>("PropertyEditor");
 	DetailsView = PropertyEditorModule.CreateDetailView(DetailsViewArgs);
@@ -304,24 +306,6 @@ TSharedPtr<IDetailsView>& FInstanceDataObjectFixupPanel::GenerateDetailsView(boo
 	{
 		RemoveCustomizationsWithLooseProperties(Instance->GetClass(), DetailsView);
 	}
-
-	for (const UObject* Object : Instances)
-	{
-		if (HasViewFlag(EViewFlags::HideLooseProperties))
-		{
-			DetailsView->RegisterInstancedCustomPropertyLayout(Object->GetClass(), FOnGetDetailCustomizationInstance::CreateLambda([]()
-			{
-				return MakeShared<FHideLoosePropertiesCustomization>();
-			}));
-		}
-		else if (HasViewFlag(EViewFlags::AllowRemapLooseProperties))
-		{
-			DetailsView->RegisterInstancedCustomPropertyLayout(Object->GetClass(), FOnGetDetailCustomizationInstance::CreateLambda([DiffPanel = SharedThis(this)]()
-			{
-				return MakeShared<FInstanceDataObjectFixupDetailCustomization>(DiffPanel);
-			}));
-		}
-	}
 	
 	DetailsView->SetObjects(Instances, true);
 	return DetailsView;
@@ -703,7 +687,7 @@ void FInstanceDataObjectFixupPanel::RedirectPropertyHelper(const FPropertyPath&
 		{
 			RedirectedPropertyTree->Move(From, To);
 		}
-		if (SourceProperty->HasMetaData(TEXT("isLoose")))
+		if (SourceProperty->HasMetaData(NAME_IsLooseMetadata))
 		{
 			SourceProperty->PropertyFlags |= CPF_Transient;
 			SourceProperty->SetMetaData(TEXT("Hidden"), TEXT("True"));
diff --git a/Engine/Source/Editor/PropertyEditor/Private/DetailFilter.h b/Engine/Source/Editor/PropertyEditor/Private/DetailFilter.h
--- a/Engine/Source/Editor/PropertyEditor/Private/DetailFilter.h
+++ b/Engine/Source/Editor/PropertyEditor/Private/DetailFilter.h
@@ -16,6 +16,7 @@ struct FDetailFilter
 		, bShowOnlyAnimated(false)
 		, bShowFavoritesCategory(false)
 		, bShowOnlyAllowed(false)
+		, bShowLooseProperties(false)
 	{}
 
 	bool IsEmptyFilter() const
@@ -27,7 +28,8 @@ struct FDetailFilter
 			&& bShowOnlyAllowed == false
 			&& bShowAllChildrenIfCategoryMatches == false
 			&& bShowOnlyKeyable == false
-			&& bShowOnlyAnimated == false;
+			&& bShowOnlyAnimated == false
+			&& bShowLooseProperties == false;
 	}
 
 	/** Any user search terms that items must match */
@@ -46,6 +48,8 @@ struct FDetailFilter
 	bool bShowFavoritesCategory;
 	/** If we should only show properties that match PropertyAllowList */
 	bool bShowOnlyAllowed;
+	/* If true, will also show loose properties */
+	bool bShowLooseProperties;
 	/** The set of allowed properties to show. */
 	TSet<FPropertyPath> PropertyAllowList;
 	/** The set of selected sections to show. If empty, all sections are shown.*/
diff --git a/Engine/Source/Editor/PropertyEditor/Private/DetailItemNode.cpp b/Engine/Source/Editor/PropertyEditor/Private/DetailItemNode.cpp
--- a/Engine/Source/Editor/PropertyEditor/Private/DetailItemNode.cpp
+++ b/Engine/Source/Editor/PropertyEditor/Private/DetailItemNode.cpp
@@ -14,6 +14,8 @@
 #include "SDetailSingleItemRow.h"
 #include "UObject/PropertyOptional.h"
 
+static const FName NAME_IsLooseMetadata = TEXT("IsLoose");
+
 FDetailItemNode::FDetailItemNode(const FDetailLayoutCustomization& InCustomization, TSharedRef<FDetailCategoryImpl> InParentCategory, TAttribute<bool> InIsParentEnabled, TSharedPtr<IDetailGroup> InParentGroup)
 	: Customization( InCustomization )
 	, ParentCategory( InParentCategory )
@@ -661,6 +663,22 @@ static bool PassesAllFilters( FDetailItemNode* ItemNode, const FDetailLayoutCust
 	};
 
 	bool bPassesAllFilters = true;
+	
+	TSharedPtr<FPropertyNode> PropertyNodePin = InCustomization.GetPropertyNode();
+	
+	if( PropertyNodePin.IsValid())
+	{
+		if (!InFilter.bShowLooseProperties)
+		{
+			if (FProperty* Property = PropertyNodePin->GetProperty())
+			{
+				if (Property->GetBoolMetaData(NAME_IsLooseMetadata))
+				{
+					return false;
+				}
+			}
+		}
+	}
 
 	if( InFilter.FilterStrings.Num() > 0 || 
 		InFilter.bShowOnlyModified == true || 
@@ -669,8 +687,6 @@ static bool PassesAllFilters( FDetailItemNode* ItemNode, const FDetailLayoutCust
 		InFilter.bShowOnlyAnimated == true)
 	{
 		const bool bSearchFilterIsEmpty = InFilter.FilterStrings.Num() == 0;
-
-		TSharedPtr<FPropertyNode> PropertyNodePin = InCustomization.GetPropertyNode();
 		
 		const bool bPassesCategoryFilter = !bSearchFilterIsEmpty && InFilter.bShowAllChildrenIfCategoryMatches ? Local::StringPassesFilter(InFilter, InCategoryName) : false;
 		const bool bPassesValueFilter = !bSearchFilterIsEmpty && Local::StringPassesFilter(InFilter, Local::GetPropertyNodeValueFilterString(InCustomization, PropertyNodePin));
diff --git a/Engine/Source/Editor/PropertyEditor/Private/SDetailsView.cpp b/Engine/Source/Editor/PropertyEditor/Private/SDetailsView.cpp
--- a/Engine/Source/Editor/PropertyEditor/Private/SDetailsView.cpp
+++ b/Engine/Source/Editor/PropertyEditor/Private/SDetailsView.cpp
@@ -75,6 +75,8 @@ void SDetailsView::Construct(const FArguments& InArgs, const FDetailsViewArgs& I
 		CurrentFilter.bShowFavoritesCategory = ViewConfig->bShowFavoritesCategory;
 	}
 
+	CurrentFilter.bShowLooseProperties = InDetailsViewArgs.bShowLooseProperties;
+
 	ColumnSizeData.SetValueColumnWidth(DetailsViewArgs.ColumnWidth);
 	ColumnSizeData.SetRightColumnMinWidth(DetailsViewArgs.RightColumnMinWidth);
 
diff --git a/Engine/Source/Editor/PropertyEditor/Public/DetailsViewArgs.h b/Engine/Source/Editor/PropertyEditor/Public/DetailsViewArgs.h
--- a/Engine/Source/Editor/PropertyEditor/Public/DetailsViewArgs.h
+++ b/Engine/Source/Editor/PropertyEditor/Public/DetailsViewArgs.h
@@ -114,6 +114,8 @@ struct FDetailsViewArgs
 	bool bShowCustomFilterOption : 1;
 	/** True if the section selector should be shown. */
 	bool bShowSectionSelector : 1;
+	/** If true, loose properties should be shown. */
+	bool bShowLooseProperties : 1;
 
 public:
 
@@ -147,6 +149,7 @@ public:
 		, bShowAnimatedPropertiesOption(true)
 		, bShowCustomFilterOption(false)
 		, bShowSectionSelector(false)
+		, bShowLooseProperties(false)
 	{
 	}
 
