[BugFix] Unmangled Property Names in IDOs and fixup tool

diff --git a/Engine/Plugins/Experimental/InstanceDataObjectFixupTool/Source/Private/InstanceDataObjectFixupDetailCustomization.cpp b/Engine/Plugins/Experimental/InstanceDataObjectFixupTool/Source/Private/InstanceDataObjectFixupDetailCustomization.cpp
--- a/Engine/Plugins/Experimental/InstanceDataObjectFixupTool/Source/Private/InstanceDataObjectFixupDetailCustomization.cpp
+++ b/Engine/Plugins/Experimental/InstanceDataObjectFixupTool/Source/Private/InstanceDataObjectFixupDetailCustomization.cpp
@@ -175,6 +175,18 @@ int32 FInstanceDataObjectNameWidgetOverride::GetNameWidgetIndex(FPropertyPath Pa
 	return DisplayRegularName;
 }
 
+static FText GetCleanVersePath(const FPropertyPath& Path)
+{
+	FString PathString = Path.ToString();
+	int32 I = PathString.Find(TEXT("__verse_0x"));
+	while (I < GetNum(PathString) && I != INDEX_NONE)
+	{
+		PathString.RemoveAt(I, 19, false);
+		I = PathString.Find(TEXT("__verse_0x"), ESearchCase::IgnoreCase, ESearchDir::FromStart, I);
+	}
+	return FText::FromString(PathString);
+}
+
 TSharedRef<SWidget> FInstanceDataObjectNameWidgetOverride::GeneratePropertyRedirectMenu(FPropertyPath Path) const
 {
 	FMenuBuilder MenuBuilder(true, nullptr);
@@ -189,7 +201,7 @@ TSharedRef<SWidget> FInstanceDataObjectNameWidgetOverride::GeneratePropertyRedir
 	MenuBuilder.BeginSection(NAME_None, LOCTEXT("ResetRedirect", "Reset"));
 	if (OriginalPath != Path || Panel->MarkedForDelete.Contains(OriginalPath))
 	{
-		FText OriginalPathText = FText::FromString(OriginalPath.ToString());
+		FText OriginalPathText = GetCleanVersePath(OriginalPath);
 		FText Tooltip = FText::Format(LOCTEXT("ResetTooltip", "Reset back to {0}"), OriginalPathText);
 		MenuBuilder.AddMenuEntry(OriginalPathText, Tooltip, FSlateIcon()
 						, FUIAction(FExecuteAction::CreateSP(Panel.Get(), &FInstanceDataObjectFixupPanel::OnRedirectProperty, Path, OriginalPath))
@@ -225,7 +237,7 @@ TSharedRef<SWidget> FInstanceDataObjectNameWidgetOverride::GeneratePropertyRedir
 			{
 				if (OptionProperty->SameType(ThisProperty))
 				{
-					FText DisplayName = FText::FromString(Option.ToString());
+					FText DisplayName = GetCleanVersePath(Option);
 					FText Tooltip = FText::Format(LOCTEXT("MovePropertyTooltip", "Move property to '{0}'"), DisplayName);
 					MenuBuilder.AddMenuEntry(DisplayName, Tooltip, FSlateIcon()
 					, FUIAction(FExecuteAction::CreateSP(Panel.Get(), &FInstanceDataObjectFixupPanel::OnRedirectProperty, Path, Option))
@@ -249,7 +261,7 @@ TSharedRef<SWidget> FInstanceDataObjectNameWidgetOverride::GeneratePropertyRedir
 			}
 			if (OptionProperty->SameType(ThisProperty))
 			{
-				FText DisplayName = FText::FromString(Option.ToString());
+				FText DisplayName = GetCleanVersePath(Option);
 				FText Tooltip = FText::Format(LOCTEXT("RenamePropertyTooltip", "Rename property to '{0}'"), DisplayName);
 				MenuBuilder.AddMenuEntry(DisplayName, Tooltip, FSlateIcon()
 				, FUIAction(FExecuteAction::CreateSP(Panel.Get(), &FInstanceDataObjectFixupPanel::OnRedirectProperty, Path, Option))
@@ -276,7 +288,7 @@ TSharedRef<SWidget> FInstanceDataObjectNameWidgetOverride::GeneratePropertyRedir
 			}
 			if (FInstanceDataObjectFixupPanel::FTypeConverter Converter = Panel->CreateTypeConverter(Path, Option))
 			{
-				FText DisplayName = FText::FromString(Option.ToString());
+				FText DisplayName = GetCleanVersePath(Option);
 				FText TypeName = FText::FromName(Option.GetLeafMostProperty().Property->GetID());
 				FText Warning = Converter.GetWarning();
 				FText Tooltip = FText::Format(LOCTEXT("ConvertTypeTooltip", "Change type to {0}"), TypeName);
@@ -310,7 +322,7 @@ TSharedRef<SWidget> FInstanceDataObjectNameWidgetOverride::GeneratePropertyRedir
 			}
 			if (FInstanceDataObjectFixupPanel::FTypeConverter Converter = Panel->CreateTypeConverter(Path, Option))
 			{
-				FText DisplayName = FText::FromString(Option.ToString());
+				FText DisplayName = GetCleanVersePath(Option);
 				FText PropDisplayName = Option.GetLeafMostProperty().Property->GetDisplayNameText();
 				FText TypeName = FText::FromName(Option.GetLeafMostProperty().Property->GetID());
 				FText Warning = Converter.GetWarning();
diff --git a/Engine/Source/Runtime/CoreUObject/Private/UObject/InstanceDataObjectUtils.cpp b/Engine/Source/Runtime/CoreUObject/Private/UObject/InstanceDataObjectUtils.cpp
--- a/Engine/Source/Runtime/CoreUObject/Private/UObject/InstanceDataObjectUtils.cpp
+++ b/Engine/Source/Runtime/CoreUObject/Private/UObject/InstanceDataObjectUtils.cpp
@@ -330,10 +330,42 @@ namespace UE
 		return Result;
 	}
 
+	static FString UnmanglePropertyName(const FName MaybeMangledName, bool* bOutNameWasMangled)
+	{
+		FString Result = MaybeMangledName.ToString();
+		if (Result.StartsWith(TEXT("__verse_0x")))
+		{
+			// chop "__verse_0x" (10 char) + CRC (8 char) + "_" (1 char)
+			Result = Result.RightChop(19);
+			if (bOutNameWasMangled)
+			{
+				*bOutNameWasMangled = true;
+			}
+		}
+		else if (bOutNameWasMangled)
+		{
+			*bOutNameWasMangled = false;
+		}
+		return Result;
+	}
+
 	// recursively re-instances all structs contained by this property to include loose properties
 	static void ConvertToInstanceDataObjectProperty(FProperty* Property, FPropertyTypeName PropertyType, UObject* Outer,
 		const TMap<FWildcardPropertyPathName, TMap<FName, const FProperty*>>& LooseProperties, FWildcardPropertyPathName& Path)
 	{
+#if WITH_EDITORONLY_DATA
+		static const FName NAME_DisplayName(TEXT("DisplayName"));
+		if (!Property->HasMetaData(NAME_DisplayName))
+		{
+			bool bNeedsDisplayName = false;
+			const FString DisplayName = UnmanglePropertyName(Property->GetFName(), &bNeedsDisplayName);
+			if (bNeedsDisplayName)
+			{
+				Property->SetMetaData(NAME_DisplayName, *DisplayName);
+			}
+		}
+#endif
+		
 		if (FStructProperty* AsStructProperty = CastField<FStructProperty>(Property))
 		{
 			if (!AsStructProperty->Struct->UseNativeSerialization())
