[CIS Fix] Fixed use of deprecated FString::RemoveAt method

diff --git a/Engine/Plugins/Experimental/InstanceDataObjectFixupTool/Source/Private/InstanceDataObjectFixupDetailCustomization.cpp b/Engine/Plugins/Experimental/InstanceDataObjectFixupTool/Source/Private/InstanceDataObjectFixupDetailCustomization.cpp
--- a/Engine/Plugins/Experimental/InstanceDataObjectFixupTool/Source/Private/InstanceDataObjectFixupDetailCustomization.cpp
+++ b/Engine/Plugins/Experimental/InstanceDataObjectFixupTool/Source/Private/InstanceDataObjectFixupDetailCustomization.cpp
@@ -181,7 +181,7 @@ static FText GetCleanVersePath(const FPropertyPath& Path)
 	int32 I = PathString.Find(TEXT("__verse_0x"));
 	while (I < GetNum(PathString) && I != INDEX_NONE)
 	{
-		PathString.RemoveAt(I, 19, false);
+		PathString.RemoveAt(I, 19, EAllowShrinking::No);
 		I = PathString.Find(TEXT("__verse_0x"), ESearchCase::IgnoreCase, ESearchDir::FromStart, I);
 	}
 	return FText::FromString(PathString);
