[Feature] Added bNeedsFixup flag to objects in the repository

diff --git a/Engine/Plugins/Experimental/InstanceDataObjectFixupTool/Source/Private/InstanceDataObjectFixupPanel.cpp b/Engine/Plugins/Experimental/InstanceDataObjectFixupTool/Source/Private/InstanceDataObjectFixupPanel.cpp
--- a/Engine/Plugins/Experimental/InstanceDataObjectFixupTool/Source/Private/InstanceDataObjectFixupPanel.cpp
+++ b/Engine/Plugins/Experimental/InstanceDataObjectFixupTool/Source/Private/InstanceDataObjectFixupPanel.cpp
@@ -19,6 +19,7 @@
 #define LOCTEXT_NAMESPACE "InstanceDataObjectFixupPanel"
 
 static const FName NAME_IsLooseMetadata(TEXT("IsLoose"));
+static const FName NAME_ContainsLoosePropertiesMetadata(ANSITEXTVIEW("ContainsLooseProperties"));
 
 FRedirectedPropertyNode::FRedirectedPropertyNode(const FRedirectedPropertyNode& Other)
 	: PropertyName(Other.PropertyName)
@@ -212,6 +213,39 @@ FInstanceDataObjectFixupPanel::FInstanceDataObjectFixupPanel(TConstArrayView<TOb
 	InitRedirectedPropertyTree();
 }
 
+static bool ObjectHasLoosePropertiesThatNeedFixup(UObject* Object)
+{
+	bool bNeedsFixup = false;
+	Object->GetClass()->Visit(Object, [&bNeedsFixup](const FPropertyVisitorPath& Path, void* Data)->EPropertyVisitorControlFlow
+	{
+		const FProperty* Property = Path.Top().Property;
+		if (!Property->HasAnyPropertyFlags(CPF_Transient) && Property->GetBoolMetaData(NAME_IsLooseMetadata))
+		{
+			bNeedsFixup = true;
+			return EPropertyVisitorControlFlow::Stop;
+		}
+		if (!Property->GetBoolMetaData(NAME_ContainsLoosePropertiesMetadata))
+		{
+			// if this sub-struct doesn't contain loose properties, it won't need fixup
+			return EPropertyVisitorControlFlow::StepOver;
+		}
+		return EPropertyVisitorControlFlow::StepInto;
+	});
+	return bNeedsFixup;
+}
+
+FInstanceDataObjectFixupPanel::~FInstanceDataObjectFixupPanel()
+{
+	for (UObject* Instance : Instances)
+	{
+		if (!ObjectHasLoosePropertiesThatNeedFixup(Instance))
+		{
+			UE::FPropertyBagRepository& Repository = UE::FPropertyBagRepository::Get();
+			Repository.MarkAsFixedUp(Repository.FindInstanceForDataObject(Instance));
+		}
+	}
+}
+
 int32 FInstanceDataObjectFixupPanel::Find(UObject* Value) const
 {
 	return Instances.Find(Value);
diff --git a/Engine/Plugins/Experimental/InstanceDataObjectFixupTool/Source/Private/InstanceDataObjectFixupPanel.h b/Engine/Plugins/Experimental/InstanceDataObjectFixupTool/Source/Private/InstanceDataObjectFixupPanel.h
--- a/Engine/Plugins/Experimental/InstanceDataObjectFixupTool/Source/Private/InstanceDataObjectFixupPanel.h
+++ b/Engine/Plugins/Experimental/InstanceDataObjectFixupTool/Source/Private/InstanceDataObjectFixupPanel.h
@@ -54,6 +54,7 @@ public:
 	};
 	
 	FInstanceDataObjectFixupPanel(TConstArrayView<TObjectPtr<UObject>> InstanceDataObjects, EViewFlags ViewFlags);
+	~FInstanceDataObjectFixupPanel();
 	
 	int32 Find(UObject* Value) const;
 	TSharedPtr<IDetailsView>& GenerateDetailsView(bool bScrollbarOnLeft = false);
diff --git a/Engine/Source/Runtime/CoreUObject/Internal/UObject/PropertyBagRepository.h b/Engine/Source/Runtime/CoreUObject/Internal/UObject/PropertyBagRepository.h
--- a/Engine/Source/Runtime/CoreUObject/Internal/UObject/PropertyBagRepository.h
+++ b/Engine/Source/Runtime/CoreUObject/Internal/UObject/PropertyBagRepository.h
@@ -23,6 +23,7 @@ class FPropertyBagRepository : public FGCObject
 
 		FPropertyPathNameTree* Tree = nullptr;
 		TObjectPtr<UObject> InstanceDataObject = nullptr;
+		bool bNeedsFixup = false;
 	};
 	// TODO: Make private throughout and extend access permissions here or in wrapper classes? Don't want engine code modifying bags outside of serializers and details panels.
 	//friend UObjectBase;
@@ -95,6 +96,8 @@ public:
 	 * @return			- Does the object's InstanceDataObject contain any loose properties requiring user fixup before the object may be published?
 	 */
 	COREUOBJECT_API bool RequiresFixup(const UObject* Object) const;
+	// set the bNeedsFixup flag for this object's IDO to false
+	COREUOBJECT_API void MarkAsFixedUp(const UObject* Object = nullptr);
 
 	// Accessors
 	COREUOBJECT_API bool HasInstanceDataObject(const UObject* Owner) const;
diff --git a/Engine/Source/Runtime/CoreUObject/Private/UObject/PropertyBagRepository.cpp b/Engine/Source/Runtime/CoreUObject/Private/UObject/PropertyBagRepository.cpp
--- a/Engine/Source/Runtime/CoreUObject/Private/UObject/PropertyBagRepository.cpp
+++ b/Engine/Source/Runtime/CoreUObject/Private/UObject/PropertyBagRepository.cpp
@@ -529,8 +529,20 @@ void FPropertyBagRepository::DestroyOuterBag(const UObject* Owner)
 bool FPropertyBagRepository::RequiresFixup(const UObject* Object) const
 {
 	FPropertyBagRepositoryLock LockRepo(this);
-	const FPropertyBagAssociationData* BagData = AssociatedData.Find(Object);
-	return BagData && BagData->Tree && !BagData->Tree->IsEmpty() && BagData->InstanceDataObject;
+	if (const FPropertyBagAssociationData* BagData = AssociatedData.Find(Object))
+	{
+		return BagData->bNeedsFixup;
+	}
+	return false;
+}
+
+void FPropertyBagRepository::MarkAsFixedUp(const UObject* Object)
+{
+	FPropertyBagRepositoryLock LockRepo(this);
+	if (FPropertyBagAssociationData* BagData = AssociatedData.Find(Object))
+	{
+		BagData->bNeedsFixup = false;
+	}
 }
 
 bool FPropertyBagRepository::RemoveAssociationUnsafe(const UObject* Owner)
@@ -606,6 +618,11 @@ void FPropertyBagRepository::CreateInstanceDataObjectUnsafe(UObject* Owner, FPro
 	// construct InstanceDataObject class
 	// TODO: should we put the InstanceDataObject or it's class in a package?
 	const UClass* InstanceDataObjectClass = CreateInstanceDataObjectClass(PropertyTree, Owner->GetClass(), GetTransientPackage());
+	
+#if WITH_EDITOR
+	static const FName NAME_ContainsLoosePropertiesMetadata(ANSITEXTVIEW("ContainsLooseProperties"));
+	BagData.bNeedsFixup = InstanceDataObjectClass->GetBoolMetaData(NAME_ContainsLoosePropertiesMetadata);
+#endif
 
 	TObjectPtr<UObject>* OuterPtr;
 	if (FPropertyBagAssociationData* OuterData = AssociatedData.Find(Owner->GetOuter()))
