[BugFix] Various UI bug fixes and updates for the InstanceDataObjectFixupTool
- changed the FRedirectedPropertyNode to use UE::FPropertyTypeName for it's types so that containers are distinguishable from one another
- made the fixup tool tabs have the name of the object they're fixing
- fixed the needs fixup warning to open several fixup tools if more than one object needs fixup rather than attempting multi-edit
- exluded VerseFunctionProperty from being made into loose properies since they can't realistically get fixed up
- fixed issue that caused incorrect properties to get marked as set by serialization during object copies

diff --git a/Engine/Plugins/Experimental/InstanceDataObjectFixupTool/Source/Private/InstanceDataObjectFixupPanel.cpp b/Engine/Plugins/Experimental/InstanceDataObjectFixupTool/Source/Private/InstanceDataObjectFixupPanel.cpp
--- a/Engine/Plugins/Experimental/InstanceDataObjectFixupTool/Source/Private/InstanceDataObjectFixupPanel.cpp
+++ b/Engine/Plugins/Experimental/InstanceDataObjectFixupTool/Source/Private/InstanceDataObjectFixupPanel.cpp
@@ -40,13 +40,15 @@ FRedirectedPropertyNode::FRedirectedPropertyNode(const FRedirectedPropertyNode&
 
 FRedirectedPropertyNode::FRedirectedPropertyNode(const FPropertyInfo& InInfo, const TWeakPtr<FRedirectedPropertyNode>& InParent)
 	: PropertyName(InInfo.Property->GetFName())
-	, Type(InInfo.Property->GetID())
 	, ArrayIndex(InInfo.ArrayIndex)
 	, Parent(InParent)
 {
+	UE::FPropertyTypeNameBuilder TypeBuilder;
+	InInfo.Property->SaveTypeName(TypeBuilder);
+	Type = TypeBuilder.Build();
 }
 
-FRedirectedPropertyNode::FRedirectedPropertyNode(FName InPropertyName, FName InType, int32 InArrayIndex, const TWeakPtr<FRedirectedPropertyNode>& InParent)
+FRedirectedPropertyNode::FRedirectedPropertyNode(FName InPropertyName, const UE::FPropertyTypeName& InType, int32 InArrayIndex, const TWeakPtr<FRedirectedPropertyNode>& InParent)
 	: PropertyName(InPropertyName)
 	, Type(InType)
 	, ArrayIndex(InArrayIndex)
@@ -78,7 +80,7 @@ TSharedPtr<FRedirectedPropertyNode> FRedirectedPropertyNode::FindOrAdd(const FPr
 	return Child;
 }
 
-TSharedPtr<FRedirectedPropertyNode> FRedirectedPropertyNode::FindOrAdd(FName ChildPropertyName, FName ChildType, int32 ChildArrayIndex)
+TSharedPtr<FRedirectedPropertyNode> FRedirectedPropertyNode::FindOrAdd(FName ChildPropertyName, const UE::FPropertyTypeName& ChildType, int32 ChildArrayIndex)
 {
 	TSharedPtr<FRedirectedPropertyNode> Child = Find(ChildPropertyName, ChildType, ChildArrayIndex);
 	if (!Child)
@@ -118,7 +120,7 @@ bool FRedirectedPropertyNode::Remove(const FPropertyInfo& ChildInfo)
 	return false;
 }
 
-bool FRedirectedPropertyNode::Remove(FName ChildPropertyName, FName ChildType, int32 ChildArrayIndex)
+bool FRedirectedPropertyNode::Remove(FName ChildPropertyName, const UE::FPropertyTypeName& ChildType, int32 ChildArrayIndex)
 {
 	const int32 Index = FindIndex(ChildPropertyName, ChildType, ChildArrayIndex);
 	if (Index != INDEX_NONE)
@@ -155,7 +157,7 @@ TSharedPtr<FRedirectedPropertyNode> FRedirectedPropertyNode::Find(const FPropert
 	return {};
 }
 
-TSharedPtr<FRedirectedPropertyNode> FRedirectedPropertyNode::Find(FName ChildPropertyName, FName ChildType, int32 ChildArrayIndex) const
+TSharedPtr<FRedirectedPropertyNode> FRedirectedPropertyNode::Find(FName ChildPropertyName, const UE::FPropertyTypeName& ChildType, int32 ChildArrayIndex) const
 {
 	const int32 Index = FindIndex(ChildPropertyName, ChildType, ChildArrayIndex);
 	if (Index != INDEX_NONE)
@@ -194,10 +196,12 @@ bool FRedirectedPropertyNode::Move(const FPropertyPath& FromPath, const FPropert
 
 int32 FRedirectedPropertyNode::FindIndex(const FPropertyInfo& ChildInfo) const
 {
-	return FindIndex(ChildInfo.Property->GetFName(), ChildInfo.Property->GetID(), ChildInfo.ArrayIndex);
+	UE::FPropertyTypeNameBuilder ChildTypeBuilder;
+	ChildInfo.Property->SaveTypeName(ChildTypeBuilder);
+	return FindIndex(ChildInfo.Property->GetFName(), ChildTypeBuilder.Build(), ChildInfo.ArrayIndex);
 }
 
-int32 FRedirectedPropertyNode::FindIndex(FName ChildPropertyName, FName ChildType, int32 ChildArrayIndex) const
+int32 FRedirectedPropertyNode::FindIndex(FName ChildPropertyName, const UE::FPropertyTypeName& ChildType, int32 ChildArrayIndex) const
 {
 	return Children.IndexOfByPredicate([ChildPropertyName, ChildType, ChildArrayIndex](const TSharedPtr<FRedirectedPropertyNode>& Child)
 	{
@@ -423,28 +427,14 @@ bool FInstanceDataObjectFixupPanel::ShouldSplitterIgnoreRow(const TWeakPtr<FDeta
 
 bool FInstanceDataObjectFixupPanel::AreAllConflictsRedirected() const
 {
-	bool bFoundConflict = false;
-	if (const TSharedPtr<FAsyncDetailViewDiff> Diff = DiffAgainstRight.Pin())
+	for (UObject* Instance : Instances)
 	{
-		Diff->ForEach(ETreeTraverseOrder::PreOrder,
-		[this, &bFoundConflict](const TUniquePtr<FAsyncDetailViewDiff::DiffNodeType>& DiffNode)->ETreeTraverseControl
+		if (ObjectHasLoosePropertiesThatNeedFixup(Instance))
 		{
-			const TSharedPtr<FDetailTreeNode> TreeNode = DiffNode->ValueA.Pin();
-			if (DiffNode->DiffResult == ETreeDiffResult::MissingFromTree2 && TreeNode)
-			{
-				if (const TSharedPtr<IPropertyHandle> Handle = TreeNode->CreatePropertyHandle())
-				{
-					if (!Handle->IsCategoryHandle() && !MarkedForDelete.Contains(*Handle->CreateFPropertyPath()))
-					{
-						bFoundConflict = true;
-						return ETreeTraverseControl::Break;
-					}
-				}
-			}
-			return ETreeTraverseControl::Continue;
-		});
+			return false;
+		}
 	}
-	return !bFoundConflict;
+	return true;
 }
 
 void FInstanceDataObjectFixupPanel::AutoApplyMarkDeletedActions()
@@ -758,15 +748,31 @@ void FInstanceDataObjectFixupPanel::RedirectPropertyHelper(const FPropertyPath&
 	
 	if (To != From)
 	{
+		auto OnHidden = [](FProperty* Property)
+		{
+			if (Property->HasMetaData(NAME_IsLooseMetadata))
+			{
+				Property->PropertyFlags |= CPF_Transient;
+				Property->SetMetaData(TEXT("Hidden"), TEXT("True"));
+				Property->SetMetaData(TEXT("Redirected"), TEXT("True"));
+			}
+		};
 		if (To.IsValid())
 		{
 			RedirectedPropertyTree->Move(From, To);
+			for (FPropertyPath Path = From; Path.IsValid(); Path = Path.TrimPath(1).Get())
+			{
+				// because RedirectedPropertyTree->Move could've removed multiple properties in the path, we need to check each of them
+				if (!MarkedForDelete.Find(Path) && RedirectedPropertyTree->Find(Path))
+				{
+					break;
+				}
+				OnHidden(Path.GetLeafMostProperty().Property.Get());
+			}
 		}
-		if (SourceProperty->HasMetaData(NAME_IsLooseMetadata))
+		else
 		{
-			SourceProperty->PropertyFlags |= CPF_Transient;
-			SourceProperty->SetMetaData(TEXT("Hidden"), TEXT("True"));
-			SourceProperty->SetMetaData(TEXT("Redirected"), TEXT("True"));
+			OnHidden(SourceProperty);
 		}
 	}
 
diff --git a/Engine/Plugins/Experimental/InstanceDataObjectFixupTool/Source/Private/InstanceDataObjectFixupPanel.h b/Engine/Plugins/Experimental/InstanceDataObjectFixupTool/Source/Private/InstanceDataObjectFixupPanel.h
--- a/Engine/Plugins/Experimental/InstanceDataObjectFixupTool/Source/Private/InstanceDataObjectFixupPanel.h
+++ b/Engine/Plugins/Experimental/InstanceDataObjectFixupTool/Source/Private/InstanceDataObjectFixupPanel.h
@@ -15,22 +15,22 @@ struct FRedirectedPropertyNode : public TSharedFromThis<FRedirectedPropertyNode>
 	FRedirectedPropertyNode() = default;
 	FRedirectedPropertyNode(const FRedirectedPropertyNode& Other);
 	FRedirectedPropertyNode(const FPropertyInfo& Info, const TWeakPtr<FRedirectedPropertyNode>& Parent);
-	FRedirectedPropertyNode(FName InPropertyName, FName InType, int32 InArrayIndex, const TWeakPtr<FRedirectedPropertyNode>& InParent);
+	FRedirectedPropertyNode(FName InPropertyName, const UE::FPropertyTypeName& InType, int32 InArrayIndex, const TWeakPtr<FRedirectedPropertyNode>& InParent);
 	TSharedPtr<FRedirectedPropertyNode> FindOrAdd(const FPropertyPath& Path, int32 PathIndex = 0);
 	TSharedPtr<FRedirectedPropertyNode> FindOrAdd(const FPropertyInfo& ChildInfo);
-	TSharedPtr<FRedirectedPropertyNode> FindOrAdd(FName ChildPropertyName, FName ChildType, int32 ChildArrayIndex = 0);
+	TSharedPtr<FRedirectedPropertyNode> FindOrAdd(FName ChildPropertyName, const UE::FPropertyTypeName& ChildType, int32 ChildArrayIndex = 0);
 	bool Remove(const FPropertyPath& Path, int32 PathIndex = 0);
 	bool Remove(const FPropertyInfo& ChildInfo);
-	bool Remove(FName ChildPropertyName, FName ChildType, int32 ChildArrayIndex = 0);
+	bool Remove(FName ChildPropertyName, const UE::FPropertyTypeName& ChildType, int32 ChildArrayIndex = 0);
 	TSharedPtr<FRedirectedPropertyNode> Find(const FPropertyPath& Path, int32 PathIndex = 0) const;
 	TSharedPtr<FRedirectedPropertyNode> Find(const FPropertyInfo& ChildInfo) const;
-	TSharedPtr<FRedirectedPropertyNode> Find(FName ChildPropertyName, FName ChildType, int32 ChildArrayIndex = 0) const;
+	TSharedPtr<FRedirectedPropertyNode> Find(FName ChildPropertyName, const UE::FPropertyTypeName& ChildType, int32 ChildArrayIndex = 0) const;
 	bool Move(const FPropertyPath& FromPath, const FPropertyPath& ToPath);
 	int32 FindIndex(const FPropertyInfo& ChildInfo) const;
-	int32 FindIndex(FName ChildPropertyName, FName ChildType, int32 ChildArrayIndex = 0) const;
+	int32 FindIndex(FName ChildPropertyName, const UE::FPropertyTypeName& ChildType, int32 ChildArrayIndex = 0) const;
 
 	FName PropertyName;
-	FName Type;
+	UE::FPropertyTypeName Type;
 	int32 ArrayIndex;
 	TWeakPtr<FRedirectedPropertyNode> Parent;
 	TArray<TSharedPtr<FRedirectedPropertyNode>> Children;
diff --git a/Engine/Plugins/Experimental/InstanceDataObjectFixupTool/Source/Private/InstanceDataObjectFixupToolModule.cpp b/Engine/Plugins/Experimental/InstanceDataObjectFixupTool/Source/Private/InstanceDataObjectFixupToolModule.cpp
--- a/Engine/Plugins/Experimental/InstanceDataObjectFixupTool/Source/Private/InstanceDataObjectFixupToolModule.cpp
+++ b/Engine/Plugins/Experimental/InstanceDataObjectFixupTool/Source/Private/InstanceDataObjectFixupToolModule.cpp
@@ -51,12 +51,42 @@ TSharedRef<SDockTab> FInstanceDataObjectFixupToolModule::CreateInstanceDataObjec
 void FInstanceDataObjectFixupToolModule::CreateInstanceDataObjectFixupDialog(
 	TConstArrayView<TObjectPtr<UObject>> InstanceDataObjects, TObjectPtr<UObject> InstanceDataObjectsOwner) const
 {
+	FText DisplayName;
+	if (InstanceDataObjects.Num() == 1)
+	{
+		DisplayName = FText::Format(LOCTEXT("ObjectFixupTabTitle", "{0} Fix-up"), FText::FromName(InstanceDataObjects[0]->GetFName()));
+	}
+	else
+	{
+		FString ClassName;
+		for (UObject* Object : InstanceDataObjects)
+		{
+			if (ClassName.IsEmpty())
+			{
+				ClassName = Object->GetClass()->GetName();
+			}
+			else if (ClassName != Object->GetName())
+			{
+				ClassName.Empty();
+				break;
+			}
+		}
+		if (ClassName.IsEmpty())
+		{
+			DisplayName = FText::Format(LOCTEXT("ObjectFixupTabTitle", "{0} Objects Fixup"), InstanceDataObjects.Num());
+		}
+		else
+		{
+			DisplayName = FText::Format(LOCTEXT("ObjectFixupTabTitle", "{0} {1} Objects Fixup"), InstanceDataObjects.Num(), FText::FromString(ClassName));
+		}
+	}
+	
 	FGlobalTabmanager::Get()->RegisterNomadTabSpawner(InstanceDataObjectFixupToolDialogName, FOnSpawnTab::CreateLambda(
 		[&InstanceDataObjects, InstanceDataObjectsOwner](const FSpawnTabArgs& TabArgs)
 		{
 			return FInstanceDataObjectFixupToolModule::Get().CreateInstanceDataObjectFixupTab(TabArgs, InstanceDataObjects, InstanceDataObjectsOwner);
 		}))
-		.SetDisplayName(LOCTEXT("ObjectFixupTabTitle", "Object Fix-up"))
+		.SetDisplayName(DisplayName)
 		.SetGroup(WorkspaceMenu::GetMenuStructure().GetToolsCategory())
 		.SetMenuType(ETabSpawnerMenuType::Hidden);
 	if (const TSharedPtr<SDockTab> DockTab = FGlobalTabmanager::Get()->TryInvokeTab(FTabId(InstanceDataObjectFixupToolDialogName)))
diff --git a/Engine/Plugins/Experimental/InstanceDataObjectFixupTool/Source/Private/InstanceDataObjectFixupToolTedsQueries.cpp b/Engine/Plugins/Experimental/InstanceDataObjectFixupTool/Source/Private/InstanceDataObjectFixupToolTedsQueries.cpp
--- a/Engine/Plugins/Experimental/InstanceDataObjectFixupTool/Source/Private/InstanceDataObjectFixupToolTedsQueries.cpp
+++ b/Engine/Plugins/Experimental/InstanceDataObjectFixupTool/Source/Private/InstanceDataObjectFixupToolTedsQueries.cpp
@@ -65,17 +65,11 @@ void UInstanceDataObjectFixupToolTedsQueryFactory::ShowFixUpTool(TypedElementDat
 		if (bRecurseIntoObject)
 		{
 			UObject* Owner = ObjectColumn->Object.Get();
-			TArray<TObjectPtr<UObject>> InstanceDataObjects;
 			UE::FPropertyBagRepository::Get().FindNestedInstanceDataObject(Owner, true,
-				[&InstanceDataObjects](UObject* NestedObject)
+				[Owner](UObject* NestedObject)
 				{
-					InstanceDataObjects.Emplace(NestedObject);
+					FInstanceDataObjectFixupToolModule::Get().CreateInstanceDataObjectFixupDialog({NestedObject}, Owner);
 				});
-
-			if (!InstanceDataObjects.IsEmpty())
-			{
-				FInstanceDataObjectFixupToolModule::Get().CreateInstanceDataObjectFixupDialog(InstanceDataObjects, Owner);
-			}
 		}
 		else
 		{
diff --git a/Engine/Source/Runtime/CoreUObject/Internal/UObject/InstanceDataObjectUtils.h b/Engine/Source/Runtime/CoreUObject/Internal/UObject/InstanceDataObjectUtils.h
--- a/Engine/Source/Runtime/CoreUObject/Internal/UObject/InstanceDataObjectUtils.h
+++ b/Engine/Source/Runtime/CoreUObject/Internal/UObject/InstanceDataObjectUtils.h
@@ -30,6 +30,7 @@ bool StructContainsLooseProperties(const UStruct* Struct);
 /** Generate a UClass that contains the union of the properties of PropertyTree and OwnerClass. */
 UClass* CreateInstanceDataObjectClass(const FPropertyPathNameTree* PropertyTree, UClass* OwnerClass, UObject* Outer);
 
+extern bool GMarkPropertiesSetBySerialization;
 /** Notify that a property in an struct was set when the struct was deserialized. */
 void MarkPropertySetBySerialization(const UStruct* Struct, void* StructData, const FProperty* Property, int32 ArrayIndex = 0);
 /** Query whether a property in the struct was set when the struct was deserialized. */
diff --git a/Engine/Source/Runtime/CoreUObject/Private/UObject/InstanceDataObjectUtils.cpp b/Engine/Source/Runtime/CoreUObject/Private/UObject/InstanceDataObjectUtils.cpp
--- a/Engine/Source/Runtime/CoreUObject/Private/UObject/InstanceDataObjectUtils.cpp
+++ b/Engine/Source/Runtime/CoreUObject/Private/UObject/InstanceDataObjectUtils.cpp
@@ -132,6 +132,7 @@ namespace UE
 	static const FName NAME_VerseClass(ANSITEXTVIEW("VerseClass"));
 	static const FName NAME_IDOMapKey(ANSITEXTVIEW("Key"));
 	static const FName NAME_IDOMapValue(ANSITEXTVIEW("Value"));
+	bool GMarkPropertiesSetBySerialization = true;
 
 	bool bEnableIDOSupport = false;
 	FAutoConsoleVariableRef EnableIDOSupportCVar(
@@ -140,6 +141,20 @@ namespace UE
 		TEXT("Allows property bags and IDOs to be created for supported classes.")
 	);
 
+	FString ExcludedLoosePropertyTypesVar = TEXT("VerseFunctionProperty");
+	FAutoConsoleVariableRef ExcludedLoosePropertyTypesCVar(
+		TEXT("IDO.ExcludedLoosePropertyTypes"),
+		ExcludedLoosePropertyTypesVar,
+		TEXT("Comma separated list of property types that will be excluded from loose properties in IDOs.")
+	);
+
+	static TSet<FString> GetExcludedLoosePropertyTypes()
+	{
+		TArray<FString> Result;
+		ExcludedLoosePropertyTypesVar.ParseIntoArray(Result, TEXT(","));
+		return TSet<FString>(Result);
+	}
+
 	bool IsInstanceDataObjectSupportEnabled(const UObject* InObject)
 	{
 		// Note: NULL is a valid (default) input here; in that case we just return the enable flag.
@@ -389,6 +404,8 @@ namespace UE
 		UStruct* Result = NewObject<UStruct>(Outer, StructClass, MakeUniqueObjectName(nullptr, StructClass, InstanceDataObjectName));
 		Result->SetSuperStruct(Super);
 
+		TSet<FString> ExcludedLoosePropertyTypes = GetExcludedLoosePropertyTypes();
+
 		// Gather "loose" properties for child Struct
 		TArray<FProperty*> LooseInstanceDataObjectProperties;
 		if (PropertyTree)
@@ -411,6 +428,11 @@ namespace UE
 					FField* Field = FField::TryConstruct(Type.GetName(), Result, Name, RF_NoFlags);
 					if (FProperty* Property = CastField<FProperty>(Field); Property && Property->LoadTypeName(Type, It.GetNode().GetTag()))
 					{
+						if (ExcludedLoosePropertyTypes.Contains(Property->GetClass()->GetName()))
+						{
+							// skip loose types that have been explicitly excluded from IDOs
+							continue;
+						}
 						MarkPropertyAsLoose(Property);
 						ConvertToInstanceDataObjectProperty(Property, Type, Result, It.GetNode().GetSubTree());
 						LooseInstanceDataObjectProperties.Add(Property);
@@ -533,6 +555,10 @@ namespace UE
 
 	void MarkPropertySetBySerialization(const UStruct* Struct, void* StructData, const FProperty* Property, int32 ArrayIndex)
 	{
+		if (!GMarkPropertiesSetBySerialization)
+		{
+			return;
+		}
 		if (const FByteProperty* ValuesSetBySerializationProperty = FindValuesSetBySerializationProperty(Struct))
 		{
 			const int32 PropertyIndex = Property->GetIndexInOwner() + ArrayIndex;
@@ -640,6 +666,13 @@ namespace UE
 				}
 				return nullptr;
 			};
+
+			// clear existing set-flags first
+			if (const FByteProperty* ValuesSetBySerializationProperty = FindValuesSetBySerializationProperty(NewAsStruct))
+			{
+				ValuesSetBySerializationProperty->InitializeValue_InContainer(NewDataPtr);
+			}
+			
 			for (const FProperty* OldSubProperty : TFieldRange<FProperty>(OldAsStruct))
 			{
 				if (const FProperty* NewSubProperty = FindMatchingProperty(NewAsStruct, OldSubProperty))
diff --git a/Engine/Source/Runtime/CoreUObject/Private/UObject/PropertyBagRepository.cpp b/Engine/Source/Runtime/CoreUObject/Private/UObject/PropertyBagRepository.cpp
--- a/Engine/Source/Runtime/CoreUObject/Private/UObject/PropertyBagRepository.cpp
+++ b/Engine/Source/Runtime/CoreUObject/Private/UObject/PropertyBagRepository.cpp
@@ -790,6 +790,8 @@ void FPropertyBagRepository::CopyTaggedProperties(const UObject* Source, UObject
 {
 	FUObjectSerializeContext* SerializeContext = FUObjectThreadContext::Get().GetSerializeContext();
 	TGuardValue<bool> ImpersonatePropertiesScope(SerializeContext->bImpersonateProperties, true);
+	// don't mark properties as set by serialization when performing copy
+	TGuardValue<bool> MarkPropertiesSetBySerializationScope(UE::GMarkPropertiesSetBySerialization, false);
 
 	TArray<uint8> Buffer;
 	Buffer.Reserve(Source->GetClass()->GetStructureSize());
