[CISFix] Build Engine Localization - Text conflict from LOCTEXT macro for namespace in InstanceDataObjectFixupToolModule.cpp

diff --git a/Engine/Plugins/Experimental/InstanceDataObjectFixupTool/Source/Private/InstanceDataObjectFixupToolModule.cpp b/Engine/Plugins/Experimental/InstanceDataObjectFixupTool/Source/Private/InstanceDataObjectFixupToolModule.cpp
--- a/Engine/Plugins/Experimental/InstanceDataObjectFixupTool/Source/Private/InstanceDataObjectFixupToolModule.cpp
+++ b/Engine/Plugins/Experimental/InstanceDataObjectFixupTool/Source/Private/InstanceDataObjectFixupToolModule.cpp
@@ -54,7 +54,7 @@ void FInstanceDataObjectFixupToolModule::CreateInstanceDataObjectFixupDialog(
 	FText DisplayName;
 	if (InstanceDataObjects.Num() == 1)
 	{
-		DisplayName = FText::Format(LOCTEXT("ObjectFixupTabTitle", "{0} Fix-up"), FText::FromName(InstanceDataObjects[0]->GetFName()));
+		DisplayName = FText::Format(LOCTEXT("SingleObjectFixupTabTitle", "{0} Fix-up"), FText::FromName(InstanceDataObjects[0]->GetFName()));
 	}
 	else
 	{
@@ -73,11 +73,11 @@ void FInstanceDataObjectFixupToolModule::CreateInstanceDataObjectFixupDialog(
 		}
 		if (ClassName.IsEmpty())
 		{
-			DisplayName = FText::Format(LOCTEXT("ObjectFixupTabTitle", "{0} Objects Fixup"), InstanceDataObjects.Num());
+			DisplayName = FText::Format(LOCTEXT("MultiEditObjectFixupTabTitle", "{0} Objects Fixup"), InstanceDataObjects.Num());
 		}
 		else
 		{
-			DisplayName = FText::Format(LOCTEXT("ObjectFixupTabTitle", "{0} {1} Objects Fixup"), InstanceDataObjects.Num(), FText::FromString(ClassName));
+			DisplayName = FText::Format(LOCTEXT("MultiEditObjectFixupTabTitleWithClass", "{0} {1} Objects Fixup"), InstanceDataObjects.Num(), FText::FromString(ClassName));
 		}
 	}
 	
