[CorruptionFix] Fixed bug that caused type conversions to corrupt data in the fixup tool
[BugFix] Fixed bug that caused sub-properties of loose struct properties to not become marked as loose as well

diff --git a/Engine/Plugins/Experimental/InstanceDataObjectFixupTool/Source/Private/InstanceDataObjectFixupPanel.cpp b/Engine/Plugins/Experimental/InstanceDataObjectFixupTool/Source/Private/InstanceDataObjectFixupPanel.cpp
--- a/Engine/Plugins/Experimental/InstanceDataObjectFixupTool/Source/Private/InstanceDataObjectFixupPanel.cpp
+++ b/Engine/Plugins/Experimental/InstanceDataObjectFixupTool/Source/Private/InstanceDataObjectFixupPanel.cpp
@@ -20,6 +20,7 @@
 #include "UObject/OverriddenPropertySet.h"
 #include "UObject/OverridableManager.h"
 #include "UObject/TextProperty.h"
+#include "UObject/UObjectThreadContext.h"
 
 #define LOCTEXT_NAMESPACE "InstanceDataObjectFixupPanel"
 
@@ -576,18 +577,22 @@ FText FInstanceDataObjectFixupPanel::FTypeConverter::GetWarning() const
 
 bool FInstanceDataObjectFixupPanel::FTypeConverter::TryConvert(FProperty* SourceProperty, const void* SourceData, FProperty* DestinationProperty, void* DestinationData)
 {
+	
+	FUObjectSerializeContext* SerializeContext = FUObjectThreadContext::Get().GetSerializeContext();
+	TGuardValue<bool> ScopedImpersonateProperties(SerializeContext->bImpersonateProperties, true);
 	TArray<uint8, TInlineAllocator<64>> Buffer;
 	TMemoryWriterBase<TInlineAllocator<64>> MemoryWriter(Buffer);
 	FStructuredArchiveFromArchive StructuredWriter(MemoryWriter);
-	SourceProperty->SerializeItem(StructuredWriter.GetSlot(), (uint8*)SourceData);
+	// todo: handle static arrays
+	FPropertyTag SourceTag(SourceProperty, 0, (uint8*)SourceData);
+	SourceTag.SerializeTaggedProperty(StructuredWriter.GetSlot(), SourceProperty, (uint8*)SourceData, nullptr);
+	
 	FMemoryReaderView MemoryReader(Buffer);
 	FStructuredArchiveFromArchive StructuredReader(MemoryReader);
 
 	// TODO: this breaks for static array elements.
 	void* DestinationContainer = static_cast<uint8*>(DestinationData) - DestinationProperty->GetOffset_ForInternal();
 
-	// todo: handle static arrays
-	FPropertyTag SourceTag(SourceProperty, 0, (uint8*)SourceData);
 
 	bool bResult = false;
 	switch(DestinationProperty->ConvertFromType(SourceTag, StructuredReader.GetSlot(), (uint8*)DestinationContainer, SourceProperty->GetOwnerStruct(), nullptr))
@@ -595,7 +600,7 @@ bool FInstanceDataObjectFixupPanel::FTypeConverter::TryConvert(FProperty* Source
 	case EConvertFromTypeResult::UseSerializeItem:
 		if (SourceProperty->GetID() == DestinationProperty->GetID())
 		{
-			SourceTag.SerializeTaggedProperty(StructuredReader.GetSlot(), DestinationProperty, (uint8*)DestinationContainer, nullptr);
+			SourceTag.SerializeTaggedProperty(StructuredReader.GetSlot(), DestinationProperty, (uint8*)DestinationData, nullptr);
 			bResult = true;
 		}
 		break;
diff --git a/Engine/Source/Runtime/CoreUObject/Private/UObject/InstanceDataObjectUtils.cpp b/Engine/Source/Runtime/CoreUObject/Private/UObject/InstanceDataObjectUtils.cpp
--- a/Engine/Source/Runtime/CoreUObject/Private/UObject/InstanceDataObjectUtils.cpp
+++ b/Engine/Source/Runtime/CoreUObject/Private/UObject/InstanceDataObjectUtils.cpp
@@ -332,6 +332,13 @@ namespace UE
 		{
 			MarkPropertyAsLoose(AsOptionalProperty->GetValueProperty());
 		}
+		else if (const FStructProperty* AsStructProperty = CastField<FStructProperty>(Property))
+		{
+			for (FProperty* InnerProperty : TFieldRange<FProperty>(AsStructProperty->Struct))
+			{
+				MarkPropertyAsLoose(InnerProperty);
+			}
+		}
 	}
 
 	// constructs an InstanceDataObject struct by merging the properties in 
