[BugFix] marked fixed up loose properties with CPF_SkipSerialization so that CPFUO doesn't copy them during reinstantiation

diff --git a/Engine/Plugins/Experimental/InstanceDataObjectFixupTool/Source/Private/InstanceDataObjectFixupPanel.cpp b/Engine/Plugins/Experimental/InstanceDataObjectFixupTool/Source/Private/InstanceDataObjectFixupPanel.cpp
--- a/Engine/Plugins/Experimental/InstanceDataObjectFixupTool/Source/Private/InstanceDataObjectFixupPanel.cpp
+++ b/Engine/Plugins/Experimental/InstanceDataObjectFixupTool/Source/Private/InstanceDataObjectFixupPanel.cpp
@@ -231,7 +231,7 @@ static bool ObjectHasLoosePropertiesThatNeedFixup(UObject* Object)
 	Object->GetClass()->Visit(Object, [&bNeedsFixup](const FPropertyVisitorPath& Path, void* Data)->EPropertyVisitorControlFlow
 	{
 		const FProperty* Property = Path.Top().Property;
-		if (!Property->HasAnyPropertyFlags(CPF_Transient) && Property->GetBoolMetaData(NAME_IsLooseMetadata))
+		if (!Property->HasAnyPropertyFlags(CPF_SkipSerialization) && Property->GetBoolMetaData(NAME_IsLooseMetadata))
 		{
 			bNeedsFixup = true;
 			return EPropertyVisitorControlFlow::Stop;
@@ -713,10 +713,10 @@ void FInstanceDataObjectFixupPanel::RedirectPropertyHelper(const FPropertyPath&
 		FromRevertInfo = *Info;
 		if (DestinationProperty)
 		{
-			if (DestinationProperty->HasAnyPropertyFlags(CPF_Transient) != Info->bWasTransient)
+			if (DestinationProperty->HasAnyPropertyFlags(CPF_SkipSerialization) != Info->bHadSkipSerialization)
 			{
-				// toggle transient flag if needed
-				DestinationProperty->PropertyFlags ^= CPF_Transient;
+				// toggle CPF_SkipSerialization flag if needed
+				DestinationProperty->PropertyFlags ^= CPF_SkipSerialization;
 			}
 			if (!Info->bWasHidden)
 			{
@@ -732,7 +732,7 @@ void FInstanceDataObjectFixupPanel::RedirectPropertyHelper(const FPropertyPath&
 			
 			ToRevertInfo = &RevertInfo.Add(To, {
 				.OriginalPath = Info->OriginalPath,
-				.bWasTransient = SourceProperty->HasAnyPropertyFlags(CPF_Transient),
+				.bHadSkipSerialization = SourceProperty->HasAnyPropertyFlags(CPF_SkipSerialization),
 				.bWasHidden = SourceProperty->HasMetaData(TEXT("Hidden"))
 			});
 		}
@@ -745,7 +745,7 @@ void FInstanceDataObjectFixupPanel::RedirectPropertyHelper(const FPropertyPath&
 		{
 			ToRevertInfo = &RevertInfo.Add(To, {
 				.OriginalPath = From,
-				.bWasTransient = SourceProperty->HasAnyPropertyFlags(CPF_Transient)
+				.bHadSkipSerialization = SourceProperty->HasAnyPropertyFlags(CPF_SkipSerialization)
 			});
 			MarkedForDelete.Remove(From);
 		}
@@ -757,7 +757,7 @@ void FInstanceDataObjectFixupPanel::RedirectPropertyHelper(const FPropertyPath&
 		{
 			if (Property->HasMetaData(NAME_IsLooseMetadata))
 			{
-				Property->PropertyFlags |= CPF_Transient;
+				Property->PropertyFlags |= CPF_SkipSerialization;
 				Property->SetMetaData(TEXT("Hidden"), TEXT("True"));
 				Property->SetMetaData(TEXT("Redirected"), TEXT("True"));
 			}
diff --git a/Engine/Plugins/Experimental/InstanceDataObjectFixupTool/Source/Private/InstanceDataObjectFixupPanel.h b/Engine/Plugins/Experimental/InstanceDataObjectFixupTool/Source/Private/InstanceDataObjectFixupPanel.h
--- a/Engine/Plugins/Experimental/InstanceDataObjectFixupTool/Source/Private/InstanceDataObjectFixupPanel.h
+++ b/Engine/Plugins/Experimental/InstanceDataObjectFixupTool/Source/Private/InstanceDataObjectFixupPanel.h
@@ -137,7 +137,7 @@ private:
 	{
 		TArray<uint8> OriginalValue;
 		FPropertyPath OriginalPath;
-		bool bWasTransient;
+		bool bHadSkipSerialization;
 		bool bWasHidden;
 	};
 
