[BugFix] propery change notifications were being sent incorrectly in the fixup tool for IDOs

diff --git a/Engine/Plugins/Experimental/InstanceDataObjectFixupTool/Source/Private/InstanceDataObjectFixupPanel.cpp b/Engine/Plugins/Experimental/InstanceDataObjectFixupTool/Source/Private/InstanceDataObjectFixupPanel.cpp
--- a/Engine/Plugins/Experimental/InstanceDataObjectFixupTool/Source/Private/InstanceDataObjectFixupPanel.cpp
+++ b/Engine/Plugins/Experimental/InstanceDataObjectFixupTool/Source/Private/InstanceDataObjectFixupPanel.cpp
@@ -818,8 +818,8 @@ void FInstanceDataObjectFixupPanel::RedirectProperty(const FPropertyPath& From,
 		FEditPropertyChain Chain;
 		TMap<FString, int32> ArrayIndices;
 		FPropertyChangedEvent ChangeEvent = ConstructChangeEventForRedirect(To, Chain, ArrayIndices);
-		FOverridableManager::Get().PreOverrideProperty(*Instance, Chain);
-		Instance->PreEditChange(ChangeEvent.Property);
+		FPropertyChangedChainEvent ChangedChainEvent(Chain, ChangeEvent);
+		Instance->PreEditChange(Chain);
 
 		if (ToRevertInfo)
 		{
@@ -848,8 +848,7 @@ void FInstanceDataObjectFixupPanel::RedirectProperty(const FPropertyPath& From,
 			SourceProperty->CopyCompleteValue(Source, FromRevertInfoItr);
 			FromRevertInfoItr += DestinationProperty->ArrayDim * DestinationProperty->ElementSize;
 		}
-		Instance->PostEditChangeProperty(ChangeEvent);
-		FOverridableManager::Get().PostOverrideProperty(*Instance, ChangeEvent, Chain);
+		Instance->PostEditChangeChainProperty(ChangedChainEvent);
 	}
 
 	GEditor->EndTransaction();
@@ -885,8 +884,7 @@ void FInstanceDataObjectFixupPanel::RedirectProperty(const FPropertyPath& From,
 		Chains.Emplace();
 		TMap<FString, int32> ArrayIndices;
 		ChangeEvents.Add(ConstructChangeEventForRedirect(To, Chains.Last(), ArrayIndices));
-		FOverridableManager::Get().PreOverrideProperty(*Instance, Chains.Last());
-		Instance->PreEditChange(ChangeEvents.Last().Property);
+		Instance->PreEditChange(Chains.Last());
 
 		if (ToRevertInfo)
 		{
@@ -912,8 +910,8 @@ void FInstanceDataObjectFixupPanel::RedirectProperty(const FPropertyPath& From,
 			SourceProperty->CopyCompleteValue(Source, FromRevertInfoItr);
 			FromRevertInfoItr += DestinationProperty->ArrayDim * DestinationProperty->ElementSize;
 		}
-		Instance->PostEditChangeProperty(ChangeEvents[I]);
-		FOverridableManager::Get().PostOverrideProperty(*Instance, ChangeEvents[I], Chains[I]);
+		FPropertyChangedChainEvent ChangedChainEvent(Chains[I], ChangeEvents[I]);
+		Instance->PostEditChangeChainProperty(ChangedChainEvent);
 	}
 
 	GEditor->EndTransaction();
