[BugFix] Gave placeholder objects a proper error message with correct user instructions instead of launching the fixup tool

diff --git a/Engine/Plugins/Experimental/InstanceDataObjectFixupTool/Source/Private/InstanceDataObjectFixupToolTedsQueries.cpp b/Engine/Plugins/Experimental/InstanceDataObjectFixupTool/Source/Private/InstanceDataObjectFixupToolTedsQueries.cpp
--- a/Engine/Plugins/Experimental/InstanceDataObjectFixupTool/Source/Private/InstanceDataObjectFixupToolTedsQueries.cpp
+++ b/Engine/Plugins/Experimental/InstanceDataObjectFixupTool/Source/Private/InstanceDataObjectFixupToolTedsQueries.cpp
@@ -11,6 +11,9 @@
 #include "Engine/World.h"
 #include "UObject/PropertyBagRepository.h"
 #include "InstanceDataObjectFixupToolModule.h"
+#include "Framework/Notifications/NotificationManager.h"
+#include "Widgets/Notifications/SNotificationList.h"
+#define LOCTEXT_NAMESPACE "FixupToolTedsQueries"
 
 void UInstanceDataObjectFixupToolTedsQueryFactory::RegisterQueries(ITypedElementDataStorageInterface& DataStorage)
 {
@@ -49,7 +52,10 @@ void UInstanceDataObjectFixupToolTedsQueryFactory::RegisterQueries(ITypedElement
 
 void UInstanceDataObjectFixupToolTedsQueryFactory::ShowFixUpToolForPlaceholders(TypedElementDataStorage::RowHandle Row)
 {
-	ShowFixUpTool(Row, false);
+	FNotificationInfo Info(LOCTEXT("PlaceholderResolutionSuggestion", "Please fix your Verse code and/or rename the Verse class back to the original name."));
+	
+	Info.ExpireDuration = 4.0f;
+	FSlateNotificationManager::Get().AddNotification(Info);
 }
 
 void UInstanceDataObjectFixupToolTedsQueryFactory::ShowFixUpToolForLooseProperties(TypedElementDataStorage::RowHandle Row)
@@ -80,3 +86,4 @@ void UInstanceDataObjectFixupToolTedsQueryFactory::ShowFixUpTool(TypedElementDat
 		}
 	}
 }
+#undef LOCTEXT_NAMESPACE
