[BugFix] Fixed tracking of initialized and serialized property state for IDOs
[BugFix] Fixed issue that caused items to not be visible in left panel of fixup tool even if serialized property tracking is properly enabled

diff --git a/Engine/Plugins/Experimental/InstanceDataObjectFixupTool/Source/Private/InstanceDataObjectFixupPanel.cpp b/Engine/Plugins/Experimental/InstanceDataObjectFixupTool/Source/Private/InstanceDataObjectFixupPanel.cpp
--- a/Engine/Plugins/Experimental/InstanceDataObjectFixupTool/Source/Private/InstanceDataObjectFixupPanel.cpp
+++ b/Engine/Plugins/Experimental/InstanceDataObjectFixupTool/Source/Private/InstanceDataObjectFixupPanel.cpp
@@ -942,7 +942,11 @@ static void InitRedirectedPropertyTreeRec(const TSharedPtr<FRedirectedPropertyNo
 		{
 			if (SerializedState.IsSet(Property, StaticArrayIndex))
 			{
-				const TSharedPtr<FRedirectedPropertyNode>& ChildNode = Node->FindOrAdd(FPropertyInfo(Property, StaticArrayIndex));
+				FPropertyInfo PropertyInfo(
+					Property, 
+					Property->ArrayDim > 1 ? StaticArrayIndex : INDEX_NONE // use INDEX_NONE for properties that aren't in arrays
+				);
+				const TSharedPtr<FRedirectedPropertyNode>& ChildNode = Node->FindOrAdd(PropertyInfo);
 				void* Value = Property->ContainerPtrToValuePtr<void>(StructValue, StaticArrayIndex);
 				InitRedirectedPropertyTreeRec(ChildNode, Property, Value, EnteredObjects);
 			}
diff --git a/Engine/Source/Runtime/CoreUObject/Private/UObject/PropertyBagRepository.cpp b/Engine/Source/Runtime/CoreUObject/Private/UObject/PropertyBagRepository.cpp
--- a/Engine/Source/Runtime/CoreUObject/Private/UObject/PropertyBagRepository.cpp
+++ b/Engine/Source/Runtime/CoreUObject/Private/UObject/PropertyBagRepository.cpp
@@ -879,8 +879,9 @@ FScopedIDOSerializationContext::FScopedIDOSerializationContext(UObject* InObject
 	SerializeContext->bTrackUnknownProperties = bCreateIDO;
 	SerializeContext->bTrackUnknownEnumNames = bCreateIDO;
 
-	SerializeContext->bTrackInitializedProperties = bIsLoading && FInitializedPropertyValueState(InObject).IsTracking();
-	SerializeContext->bTrackSerializedProperties = bIsLoading && FSerializedPropertyValueState(InObject).IsTracking();
+	UObject* TargetObject = SerializeContext->bImpersonateProperties ? ResolveInstanceDataObject(InObject) : InObject;
+	SerializeContext->bTrackInitializedProperties = bIsLoading && FInitializedPropertyValueState(TargetObject).IsTracking();
+	SerializeContext->bTrackSerializedProperties = bIsLoading && FSerializedPropertyValueState(TargetObject).IsTracking();
 
 	if (!bIsLoading)
 	{
