[BugFix] the fixup tool handled UObject references slightly wrong. This caused a crash when remapping object properties

diff --git a/Engine/Plugins/Experimental/InstanceDataObjectFixupTool/Source/Private/InstanceDataObjectFixupDetailCustomization.cpp b/Engine/Plugins/Experimental/InstanceDataObjectFixupTool/Source/Private/InstanceDataObjectFixupDetailCustomization.cpp
--- a/Engine/Plugins/Experimental/InstanceDataObjectFixupTool/Source/Private/InstanceDataObjectFixupDetailCustomization.cpp
+++ b/Engine/Plugins/Experimental/InstanceDataObjectFixupTool/Source/Private/InstanceDataObjectFixupDetailCustomization.cpp
@@ -115,13 +115,15 @@ void FInstanceDataObjectNameWidgetOverride::GetRedirectOptions(const FProperty*
 	{
 		if (AsObjectProperty->HasAnyPropertyFlags(CPF_InstancedReference))
 		{
-			TObjectPtr<UObject> Object = AsObjectProperty->GetObjectPropertyValue(Value);
-			UE::FPropertyBagRepository& PropertyBagRepository = UE::FPropertyBagRepository::Get();
-			if (UObject* Found = PropertyBagRepository.FindInstanceDataObject(Object))
+			if (TObjectPtr<UObject> Object = AsObjectProperty->GetObjectPropertyValue(Value))
 			{
-				Object = Found;
+				UE::FPropertyBagRepository& PropertyBagRepository = UE::FPropertyBagRepository::Get();
+				if (UObject* Found = PropertyBagRepository.FindInstanceDataObject(Object))
+				{
+					Object = Found;
+				}
+				GetRedirectOptions(Object->GetClass(), Object, Path, OutPaths);
 			}
-			GetRedirectOptions(Object->GetClass(), Object, Path, OutPaths);
 		}
 	}
 	else if (const FArrayProperty* AsArrayProperty = CastField<FArrayProperty>(Property))
diff --git a/Engine/Plugins/Experimental/InstanceDataObjectFixupTool/Source/Private/InstanceDataObjectFixupPanel.cpp b/Engine/Plugins/Experimental/InstanceDataObjectFixupTool/Source/Private/InstanceDataObjectFixupPanel.cpp
--- a/Engine/Plugins/Experimental/InstanceDataObjectFixupTool/Source/Private/InstanceDataObjectFixupPanel.cpp
+++ b/Engine/Plugins/Experimental/InstanceDataObjectFixupTool/Source/Private/InstanceDataObjectFixupPanel.cpp
@@ -17,6 +17,8 @@
 #include "Elements/Common/EditorDataStorageFeatures.h"
 #include "Elements/Interfaces/TypedElementDataStorageInterface.h"
 #include "Elements/Interfaces/TypedElementDataStorageCompatibilityInterface.h"
+#include "Serialization/ObjectReader.h"
+#include "Serialization/ObjectWriter.h"
 
 #include "UObject/OverriddenPropertySet.h"
 #include "UObject/OverridableManager.h"
@@ -475,6 +477,7 @@ bool FInstanceDataObjectFixupPanel::HasViewFlag(EViewFlags Flag)
 
 static void* ResolvePath(const FPropertyPath& Path, void* Value)
 {
+	const int32 LastPathIndex = Path.GetNumProperties() - 1;
 	for(int32 PathIndex = 0; PathIndex < Path.GetNumProperties(); ++PathIndex)
 	{
 		const FPropertyInfo& PropertyInfo = Path.GetPropertyInfo(PathIndex);
@@ -486,18 +489,18 @@ static void* ResolvePath(const FPropertyPath& Path, void* Value)
 
 		Value = Property->ContainerPtrToValuePtr<void>(Value, PropertyInfo.ArrayIndex != INDEX_NONE ? PropertyInfo.ArrayIndex : 0);
 
-		if (const FObjectProperty* AsObjectProperty = CastField<FObjectProperty>(Property))
+		if (PathIndex < LastPathIndex)
 		{
-			UObject* Object = AsObjectProperty->GetObjectPropertyValue(Value);
-			UE::FPropertyBagRepository& PropertyBagRepository = UE::FPropertyBagRepository::Get();
-			if (UObject* Found = PropertyBagRepository.FindInstanceDataObject(Object))
+			if (const FObjectProperty* AsObjectProperty = CastField<FObjectProperty>(Property))
 			{
-				Object = Found;
+				UObject* Object = AsObjectProperty->GetObjectPropertyValue(Value);
+				UE::FPropertyBagRepository& PropertyBagRepository = UE::FPropertyBagRepository::Get();
+				if (UObject* Found = PropertyBagRepository.FindInstanceDataObject(Object))
+				{
+					Object = Found;
+				}
+				Value = Object;
 			}
-			Value = Object;
-		}
-		else if (PathIndex + 1 < Path.GetNumProperties())
-		{
 			if (const FArrayProperty* AsArrayProperty = CastField<FArrayProperty>(Property))
 			{
 				FScriptArrayHelper Helper(AsArrayProperty, Value);
@@ -583,15 +586,15 @@ bool FInstanceDataObjectFixupPanel::FTypeConverter::TryConvert(FProperty* Source
 	
 	FUObjectSerializeContext* SerializeContext = FUObjectThreadContext::Get().GetSerializeContext();
 	TGuardValue<bool> ScopedImpersonateProperties(SerializeContext->bImpersonateProperties, true);
-	TArray<uint8, TInlineAllocator<64>> Buffer;
-	TMemoryWriterBase<TInlineAllocator<64>> MemoryWriter(Buffer);
-	FStructuredArchiveFromArchive StructuredWriter(MemoryWriter);
+	TArray<uint8> Buffer;
+	FObjectWriter ObjectWriter(Buffer); // using FObjectWriter so that object properties will serialize correctly
+	FStructuredArchiveFromArchive StructuredWriter(ObjectWriter);
 	// todo: handle static arrays
 	FPropertyTag SourceTag(SourceProperty, 0, (uint8*)SourceData);
 	SourceTag.SerializeTaggedProperty(StructuredWriter.GetSlot(), SourceProperty, (uint8*)SourceData, nullptr);
 	
-	FMemoryReaderView MemoryReader(Buffer);
-	FStructuredArchiveFromArchive StructuredReader(MemoryReader);
+	FObjectReader ObjectReader(Buffer);
+	FStructuredArchiveFromArchive StructuredReader(ObjectReader);
 
 	// TODO: this breaks for static array elements.
 	void* DestinationContainer = static_cast<uint8*>(DestinationData) - DestinationProperty->GetOffset_ForInternal();
