[Feature] Refactored DataRecoveryTool to write IDTs
[Feature] Work on IDT system API so that tools can author/edit new IDT scripts
[Minor]  FPropertyBagRepository::RequiresFixup now simply queries whether there's an unknown property tree. This saves us having to keep the flag updated

diff --git a/Engine/Plugins/Experimental/InstanceDataObjectFixupTool/Source/InstanceDataObjectFixupTool.Build.cs b/Engine/Plugins/Experimental/InstanceDataObjectFixupTool/Source/InstanceDataObjectFixupTool.Build.cs
--- a/Engine/Plugins/Experimental/InstanceDataObjectFixupTool/Source/InstanceDataObjectFixupTool.Build.cs
+++ b/Engine/Plugins/Experimental/InstanceDataObjectFixupTool/Source/InstanceDataObjectFixupTool.Build.cs
@@ -22,7 +22,8 @@ namespace UnrealBuildTool.Rules
 					"WorkspaceMenuStructure", 
 					"UnrealEd",
 					"TedsAlerts",
-					"TypedElementFramework"
+					"TypedElementFramework",
+					"Json"
 				}
 			);
 		}
diff --git a/Engine/Plugins/Experimental/InstanceDataObjectFixupTool/Source/Private/InstanceDataObjectFixupDetailCustomization.cpp b/Engine/Plugins/Experimental/InstanceDataObjectFixupTool/Source/Private/InstanceDataObjectFixupDetailCustomization.cpp
--- a/Engine/Plugins/Experimental/InstanceDataObjectFixupTool/Source/Private/InstanceDataObjectFixupDetailCustomization.cpp
+++ b/Engine/Plugins/Experimental/InstanceDataObjectFixupTool/Source/Private/InstanceDataObjectFixupDetailCustomization.cpp
@@ -102,7 +102,12 @@ void FInstanceDataObjectNameWidgetOverride::GetRedirectOptions(const FProperty*
 		return;
 	}
 
-	if (!DiffPanel.Pin()->RedirectedPropertyTree->Find(Path))
+	if (!DiffPanel.IsValid())
+	{
+		return;
+	}
+
+	if (!DiffPanel.Pin()->PropertyNeedsRedirect(Path))
 	{
 		OutPaths.Add(Path);
 	}
@@ -189,6 +194,20 @@ static FText GetCleanVersePath(const FPropertyPath& Path)
 	return FText::FromString(PathString);
 }
 
+static bool DoesRedirectRespectContainerIndices(const FPropertyPath& Src, const FPropertyPath& Dst)
+{
+	for (int Index = 0; Index < Src.GetNumProperties() && Index < Dst.GetNumProperties(); ++Index)
+	{
+		const FPropertyInfo& SrcInfo = Src.GetPropertyInfo(Index);
+		const FPropertyInfo& DstInfo = Dst.GetPropertyInfo(Index);
+		if (SrcInfo.ArrayIndex != DstInfo.ArrayIndex)
+		{
+			return false;
+		}
+	}
+	return true;
+}
+
 TSharedRef<SWidget> FInstanceDataObjectNameWidgetOverride::GeneratePropertyRedirectMenu(FPropertyPath Path) const
 {
 	FMenuBuilder MenuBuilder(true, nullptr);
@@ -226,13 +245,21 @@ TSharedRef<SWidget> FInstanceDataObjectNameWidgetOverride::GeneratePropertyRedir
 	}
 	MenuBuilder.EndSection();
 
-	UObject* FirstInstanceDataObject = Panel->Instances[0];
-	TSet<FPropertyPath> RedirectOptions = GetRedirectOptions(FirstInstanceDataObject->GetClass(), FirstInstanceDataObject);
+	TSet<FPropertyPath> RedirectOptions;
+	if (!Panel->Instances.IsEmpty())
+	{
+		UObject* FirstInstanceDataObject = Panel->Instances[0];
+		RedirectOptions = GetRedirectOptions(FirstInstanceDataObject->GetClass(), FirstInstanceDataObject);
+	}
 
 	MenuBuilder.BeginSection(NAME_None, LOCTEXT("MoveProperty", "Move"));
 	{
 		for (const FPropertyPath& Option : RedirectOptions)
 		{
+			if (!DoesRedirectRespectContainerIndices(Path, Option))
+			{
+				continue;
+			}
 			const FProperty* ThisProperty = Path.GetLeafMostProperty().Property.Get();
 			const FProperty* OptionProperty = Option.GetLeafMostProperty().Property.Get();
 			if (ThisProperty->GetFName() == OptionProperty->GetFName())
@@ -255,6 +282,10 @@ TSharedRef<SWidget> FInstanceDataObjectNameWidgetOverride::GeneratePropertyRedir
 	{
 		for (const FPropertyPath& Option : RedirectOptions)
 		{
+			if (!DoesRedirectRespectContainerIndices(Path, Option))
+			{
+				continue;
+			}
 			const FProperty* ThisProperty = Path.GetLeafMostProperty().Property.Get();
 			const FProperty* OptionProperty = Option.GetLeafMostProperty().Property.Get();
 			if (ThisProperty->GetFName() == OptionProperty->GetFName())
@@ -274,73 +305,73 @@ TSharedRef<SWidget> FInstanceDataObjectNameWidgetOverride::GeneratePropertyRedir
 	}
 	MenuBuilder.EndSection();
 
-	MenuBuilder.BeginSection(NAME_None, LOCTEXT("ChangeToDifferingTypeConversion", "Convert Type"));
-	{
-		for (const FPropertyPath& Option : RedirectOptions)
-		{
-			const FProperty* ThisProperty = Path.GetLeafMostProperty().Property.Get();
-			const FProperty* OptionProperty = Option.GetLeafMostProperty().Property.Get();
-			if (OptionProperty->SameType(ThisProperty))
-			{
-				continue; // same type handled above
-			}
-			if (ThisProperty->GetFName() != OptionProperty->GetFName())
-			{
-				continue; // renames to handled below
-			}
-			if (FInstanceDataObjectFixupPanel::FTypeConverter Converter = Panel->CreateTypeConverter(Path, Option))
-			{
-				FText DisplayName = GetCleanVersePath(Option);
-				FText TypeName = FText::FromName(Option.GetLeafMostProperty().Property->GetID());
-				FText Warning = Converter.GetWarning();
-				FText Tooltip = FText::Format(LOCTEXT("ConvertTypeTooltip", "Change type to {0}"), TypeName);
-				if (!Warning.IsEmpty())
-				{
-					DisplayName = FText::Format(LOCTEXT("ConvertTypeDisplayNameWithWarning", "⚠{0}"), DisplayName);
-					Tooltip = FText::Format(LOCTEXT("ConvertTypeTooltipWithWarning", "⚠{0}\n Warning: {1}"), Tooltip, Warning);
-				}
-				MenuBuilder.AddMenuEntry(DisplayName, Tooltip, FSlateIcon()
-				, FUIAction(FExecuteAction::CreateSP(Panel.Get(), &FInstanceDataObjectFixupPanel::OnRedirectProperty, Path, Option, Converter))
-				, NAME_None
-				, EUserInterfaceActionType::RadioButton);
-			}
-		}
-	}
-	MenuBuilder.EndSection();
+	//MenuBuilder.BeginSection(NAME_None, LOCTEXT("ChangeToDifferingTypeConversion", "Convert Type"));
+	//{
+	//	for (const FPropertyPath& Option : RedirectOptions)
+	//	{
+	//		const FProperty* ThisProperty = Path.GetLeafMostProperty().Property.Get();
+	//		const FProperty* OptionProperty = Option.GetLeafMostProperty().Property.Get();
+	//		if (OptionProperty->SameType(ThisProperty))
+	//		{
+	//			continue; // same type handled above
+	//		}
+	//		if (ThisProperty->GetFName() != OptionProperty->GetFName())
+	//		{
+	//			continue; // renames to handled below
+	//		}
+	//		if (FInstanceDataObjectFixupPanel::FTypeConverter Converter = Panel->CreateTypeConverter(Path, Option))
+	//		{
+	//			FText DisplayName = GetCleanVersePath(Option);
+	//			FText TypeName = FText::FromName(Option.GetLeafMostProperty().Property->GetID());
+	//			FText Warning = Converter.GetWarning();
+	//			FText Tooltip = FText::Format(LOCTEXT("ConvertTypeTooltip", "Change type to {0}"), TypeName);
+	//			if (!Warning.IsEmpty())
+	//			{
+	//				DisplayName = FText::Format(LOCTEXT("ConvertTypeDisplayNameWithWarning", "⚠{0}"), DisplayName);
+	//				Tooltip = FText::Format(LOCTEXT("ConvertTypeTooltipWithWarning", "⚠{0}\n Warning: {1}"), Tooltip, Warning);
+	//			}
+	//			MenuBuilder.AddMenuEntry(DisplayName, Tooltip, FSlateIcon()
+	//			, FUIAction(FExecuteAction::CreateSP(Panel.Get(), &FInstanceDataObjectFixupPanel::OnRedirectProperty, Path, Option, Converter))
+	//			, NAME_None
+	//			, EUserInterfaceActionType::RadioButton);
+	//		}
+	//	}
+	//}
+	//MenuBuilder.EndSection();
 
-	MenuBuilder.BeginSection(NAME_None, LOCTEXT("RenameToDifferingTypeConversion", "Convert Type and Rename"));
-	{
-		for (const FPropertyPath& Option : RedirectOptions)
-		{
-			const FProperty* ThisProperty = Path.GetLeafMostProperty().Property.Get();
-			const FProperty* OptionProperty = Option.GetLeafMostProperty().Property.Get();
-			if (OptionProperty->SameType(ThisProperty))
-			{
-				continue; // same type handled above
-			}
-			if (ThisProperty->GetFName() == OptionProperty->GetFName())
-			{
-				continue; // renames to handled below
-			}
-			if (FInstanceDataObjectFixupPanel::FTypeConverter Converter = Panel->CreateTypeConverter(Path, Option))
-			{
-				FText DisplayName = GetCleanVersePath(Option);
-				FText PropDisplayName = Option.GetLeafMostProperty().Property->GetDisplayNameText();
-				FText TypeName = FText::FromName(Option.GetLeafMostProperty().Property->GetID());
-				FText Warning = Converter.GetWarning();
-				FText Tooltip = FText::Format(LOCTEXT("ConvertTypeAndRenameTooltip", "Change type to {0} and rename to {1}"), TypeName, PropDisplayName);
-				if (!Warning.IsEmpty())
-				{
-					DisplayName = FText::Format(LOCTEXT("ConvertTypeAndRenameDisplayNameWithWarning", "⚠{0}"), DisplayName);
-					Tooltip = FText::Format(LOCTEXT("ConvertTypeAndRenameTooltipWithWarning", "⚠{0}\n Warning: {1}"), Tooltip, Warning);
-				}
-				MenuBuilder.AddMenuEntry(DisplayName, Tooltip, FSlateIcon()
-				, FUIAction(FExecuteAction::CreateSP(Panel.Get(), &FInstanceDataObjectFixupPanel::OnRedirectProperty, Path, Option, Converter))
-				, NAME_None
-				, EUserInterfaceActionType::RadioButton);
-			}
-		}
-	}
+	//MenuBuilder.BeginSection(NAME_None, LOCTEXT("RenameToDifferingTypeConversion", "Convert Type and Rename"));
+	//{
+	//	for (const FPropertyPath& Option : RedirectOptions)
+	//	{
+	//		const FProperty* ThisProperty = Path.GetLeafMostProperty().Property.Get();
+	//		const FProperty* OptionProperty = Option.GetLeafMostProperty().Property.Get();
+	//		if (OptionProperty->SameType(ThisProperty))
+	//		{
+	//			continue; // same type handled above
+	//		}
+	//		if (ThisProperty->GetFName() == OptionProperty->GetFName())
+	//		{
+	//			continue; // renames to handled below
+	//		}
+	//		if (FInstanceDataObjectFixupPanel::FTypeConverter Converter = Panel->CreateTypeConverter(Path, Option))
+	//		{
+	//			FText DisplayName = GetCleanVersePath(Option);
+	//			FText PropDisplayName = Option.GetLeafMostProperty().Property->GetDisplayNameText();
+	//			FText TypeName = FText::FromName(Option.GetLeafMostProperty().Property->GetID());
+	//			FText Warning = Converter.GetWarning();
+	//			FText Tooltip = FText::Format(LOCTEXT("ConvertTypeAndRenameTooltip", "Change type to {0} and rename to {1}"), TypeName, PropDisplayName);
+	//			if (!Warning.IsEmpty())
+	//			{
+	//				DisplayName = FText::Format(LOCTEXT("ConvertTypeAndRenameDisplayNameWithWarning", "⚠{0}"), DisplayName);
+	//				Tooltip = FText::Format(LOCTEXT("ConvertTypeAndRenameTooltipWithWarning", "⚠{0}\n Warning: {1}"), Tooltip, Warning);
+	//			}
+	//			MenuBuilder.AddMenuEntry(DisplayName, Tooltip, FSlateIcon()
+	//			, FUIAction(FExecuteAction::CreateSP(Panel.Get(), &FInstanceDataObjectFixupPanel::OnRedirectProperty, Path, Option, Converter))
+	//			, NAME_None
+	//			, EUserInterfaceActionType::RadioButton);
+	//		}
+	//	}
+	//}
 	MenuBuilder.EndSection();
 	
 	return MenuBuilder.MakeWidget();
diff --git a/Engine/Plugins/Experimental/InstanceDataObjectFixupTool/Source/Private/InstanceDataObjectFixupPanel.cpp b/Engine/Plugins/Experimental/InstanceDataObjectFixupTool/Source/Private/InstanceDataObjectFixupPanel.cpp
--- a/Engine/Plugins/Experimental/InstanceDataObjectFixupTool/Source/Private/InstanceDataObjectFixupPanel.cpp
+++ b/Engine/Plugins/Experimental/InstanceDataObjectFixupTool/Source/Private/InstanceDataObjectFixupPanel.cpp
@@ -20,238 +20,62 @@
 #include "Elements/Interfaces/TypedElementDataStorageCompatibilityInterface.h"
 #include "Serialization/ObjectReader.h"
 #include "Serialization/ObjectWriter.h"
+#include "Dom/JsonValue.h"
+#include "Dom/JsonObject.h"
 
 #include "UObject/OverriddenPropertySet.h"
 #include "UObject/OverridableManager.h"
 #include "UObject/TextProperty.h"
 #include "UObject/UObjectThreadContext.h"
 
-#include UE_INLINE_GENERATED_CPP_BY_NAME(InstanceDataObjectFixupPanel)
+#include "UObject/PropertyPathNameTree.h"
+#include "UObject/InstanceDataTransforms.h"
 
 #define LOCTEXT_NAMESPACE "InstanceDataObjectFixupPanel"
 
 static const FName NAME_IsLooseMetadata(TEXT("IsLoose"));
 static const FName NAME_ContainsLoosePropertiesMetadata(ANSITEXTVIEW("ContainsLooseProperties"));
 
-FRedirectedPropertyNode::FRedirectedPropertyNode(const FRedirectedPropertyNode& Other)
-	: PropertyName(Other.PropertyName)
-	, Type(Other.Type)
-	, ArrayIndex(Other.ArrayIndex)
-{
-	// deep copy tree
-	for (const TSharedPtr<FRedirectedPropertyNode>& Child : Other.Children)
-	{
-		Children.Add(MakeShared<FRedirectedPropertyNode>(*Child));
-	}
-}
-
-FRedirectedPropertyNode::FRedirectedPropertyNode(const FPropertyInfo& InInfo, const TWeakPtr<FRedirectedPropertyNode>& InParent)
-	: PropertyName(InInfo.Property->GetFName())
-	, ArrayIndex(InInfo.ArrayIndex)
-	, Parent(InParent)
-{
-	UE::FPropertyTypeNameBuilder TypeBuilder;
-	InInfo.Property->SaveTypeName(TypeBuilder);
-	Type = TypeBuilder.Build();
-}
-
-FRedirectedPropertyNode::FRedirectedPropertyNode(FName InPropertyName, const UE::FPropertyTypeName& InType, int32 InArrayIndex, const TWeakPtr<FRedirectedPropertyNode>& InParent)
-	: PropertyName(InPropertyName)
-	, Type(InType)
-	, ArrayIndex(InArrayIndex)
-	, Parent(InParent)
-{
-}
-
-TSharedPtr<FRedirectedPropertyNode> FRedirectedPropertyNode::FindOrAdd(const FPropertyPath& Path, int32 PathIndex)
-{
-	check(PathIndex <= Path.GetNumProperties());
-	if (PathIndex == Path.GetNumProperties())
-	{
-		return SharedThis(this);
-	}
-	
-	const FPropertyInfo& ChildInfo = Path.GetPropertyInfo(PathIndex);
-	const TSharedPtr<FRedirectedPropertyNode> Child = FindOrAdd(ChildInfo);
-	return Child->FindOrAdd(Path, PathIndex + 1);
-}
-
-TSharedPtr<FRedirectedPropertyNode> FRedirectedPropertyNode::FindOrAdd(const FPropertyInfo& ChildInfo)
-{
-	TSharedPtr<FRedirectedPropertyNode> Child = Find(ChildInfo);
-	if (!Child)
-	{
-		Child = MakeShared<FRedirectedPropertyNode>(ChildInfo, SharedThis(this));
-		Children.Add(Child);
-	}
-	return Child;
-}
 
-TSharedPtr<FRedirectedPropertyNode> FRedirectedPropertyNode::FindOrAdd(FName ChildPropertyName, const UE::FPropertyTypeName& ChildType, int32 ChildArrayIndex)
+static bool IsPathOverridden(UObject* Object, const FPropertyPath& Path)
 {
-	TSharedPtr<FRedirectedPropertyNode> Child = Find(ChildPropertyName, ChildType, ChildArrayIndex);
-	if (!Child)
-	{
-		Child = MakeShared<FRedirectedPropertyNode>(ChildPropertyName, ChildType, ChildArrayIndex, SharedThis(this));
-		Children.Add(Child);
-	}
-	return Child;
-}
-
-bool FRedirectedPropertyNode::Remove(const FPropertyPath& Path, int32 PathIndex)
-{
-	if (TSharedPtr<FRedirectedPropertyNode> NodeToRemove = Find(Path, PathIndex))
+	FPropertyVisitorPath VisitorPath;
+	for (int32 I = 0; I < Path.GetNumProperties(); ++I)
 	{
-		do
+		FProperty* Property = Path.GetPropertyInfo(I).Property.Get();
+		int32 Index = Path.GetPropertyInfo(I).ArrayIndex;
+		EPropertyVisitorInfoType SegmentType = (Index != INDEX_NONE) ? EPropertyVisitorInfoType::StaticArrayIndex : EPropertyVisitorInfoType::None;
+		if (I > 0 && Index != INDEX_NONE)
 		{
-			const TSharedPtr<FRedirectedPropertyNode> ParentNode = NodeToRemove->Parent.Pin();
-			if (ParentNode)
+			FProperty* ParentProperty = Path.GetPropertyInfo(I - 1).Property.Get();
+			if (ParentProperty->IsA<FArrayProperty>() || ParentProperty->IsA<FSetProperty>())
 			{
-				ParentNode->Remove(NodeToRemove->PropertyName, NodeToRemove->Type, NodeToRemove->ArrayIndex);
+				SegmentType = EPropertyVisitorInfoType::ContainerIndex;
 			}
-			NodeToRemove = ParentNode;
-		} while (NodeToRemove && NodeToRemove->Children.IsEmpty());
-		return true;
-	}
-	return false;
-}
-
-bool FRedirectedPropertyNode::Remove(const FPropertyInfo& ChildInfo)
-{
-	const int32 Index = FindIndex(ChildInfo);
-	if (Index != INDEX_NONE)
-	{
-		Children.RemoveAt(Index);
-		return true;
-	}
-	return false;
-}
-
-bool FRedirectedPropertyNode::Remove(FName ChildPropertyName, const UE::FPropertyTypeName& ChildType, int32 ChildArrayIndex)
-{
-	const int32 Index = FindIndex(ChildPropertyName, ChildType, ChildArrayIndex);
-	if (Index != INDEX_NONE)
-	{
-		Children.RemoveAt(Index);
-		return true;
-	}
-	return false;
-}
-
-TSharedPtr<FRedirectedPropertyNode> FRedirectedPropertyNode::Find(const FPropertyPath& Path, int32 PathIndex) const
-{
-	check(PathIndex <= Path.GetNumProperties());
-	if (PathIndex == Path.GetNumProperties())
-	{
-		return SharedThis(const_cast<FRedirectedPropertyNode*>(this));
-	}
-	
-	const FPropertyInfo& ChildInfo = Path.GetPropertyInfo(PathIndex);
-	if (const TSharedPtr<FRedirectedPropertyNode> Child = Find(ChildInfo))
-	{
-		return Child->Find(Path, PathIndex + 1);
-	}
-	return {};
-}
-
-TSharedPtr<FRedirectedPropertyNode> FRedirectedPropertyNode::Find(const FPropertyInfo& ChildInfo) const
-{
-	const int32 Index = FindIndex(ChildInfo);
-	if (Index != INDEX_NONE)
-	{
-		return Children[Index];
-	}
-	return {};
-}
-
-TSharedPtr<FRedirectedPropertyNode> FRedirectedPropertyNode::Find(FName ChildPropertyName, const UE::FPropertyTypeName& ChildType, int32 ChildArrayIndex) const
-{
-	const int32 Index = FindIndex(ChildPropertyName, ChildType, ChildArrayIndex);
-	if (Index != INDEX_NONE)
-	{
-		return Children[Index];
-	}
-	return {};
-}
-
-bool FRedirectedPropertyNode::Move(const FPropertyPath& FromPath, const FPropertyPath& ToPath)
-{
-	if (TSharedPtr<FRedirectedPropertyNode> NodeToMove = Find(FromPath))
-	{
-		const TSharedPtr<FRedirectedPropertyNode> Added = FindOrAdd(ToPath);
-
-		// reparent children
-		Added->Children = MoveTemp(NodeToMove->Children);
-		for (const TSharedPtr<FRedirectedPropertyNode>& Child : Added->Children)
-		{
-			Child->Parent = Added;
-		}
-
-		do
-		{
-			const TSharedPtr<FRedirectedPropertyNode> ParentNode = NodeToMove->Parent.Pin();
-			if (ParentNode)
+			else if (FMapProperty* MapProperty = CastField<FMapProperty>(ParentProperty))
 			{
-				ParentNode->Remove(NodeToMove->PropertyName, NodeToMove->Type, NodeToMove->ArrayIndex);
+				if (Property == MapProperty->KeyProp)
+				{
+					SegmentType = EPropertyVisitorInfoType::MapKey;
+				}
+				else if (Property == MapProperty->ValueProp)
+				{
+					SegmentType = EPropertyVisitorInfoType::MapValue;
+				}
 			}
-			NodeToMove = ParentNode;
-		} while (NodeToMove && NodeToMove->Children.IsEmpty());
-		return true;
-	}
-	return false;
-}
-
-int32 FRedirectedPropertyNode::FindIndex(const FPropertyInfo& ChildInfo) const
-{
-	UE::FPropertyTypeNameBuilder ChildTypeBuilder;
-	ChildInfo.Property->SaveTypeName(ChildTypeBuilder);
-	return FindIndex(ChildInfo.Property->GetFName(), ChildTypeBuilder.Build(), ChildInfo.ArrayIndex);
-}
-
-int32 FRedirectedPropertyNode::FindIndex(FName ChildPropertyName, const UE::FPropertyTypeName& ChildType, int32 ChildArrayIndex) const
-{
-	return Children.IndexOfByPredicate([ChildPropertyName, ChildType, ChildArrayIndex](const TSharedPtr<FRedirectedPropertyNode>& Child)
-	{
-		if (Child->ArrayIndex != INDEX_NONE)
-		{
-			// a matching index will always match regardless of type and name
-			return Child->ArrayIndex == ChildArrayIndex;
 		}
-		return Child->Type == ChildType && Child->PropertyName == ChildPropertyName;
-	});
+		VisitorPath.Push(FPropertyVisitorInfo(Property, Index, SegmentType));
+	}
+	EOverriddenPropertyOperation Override = FOverridableManager::Get().GetOverriddenPropertyOperation(Object, VisitorPath);
+	return (Override != EOverriddenPropertyOperation::None);
 }
 
 FInstanceDataObjectFixupPanel::FInstanceDataObjectFixupPanel(
 	TConstArrayView<TObjectPtr<UObject>> InstanceDataObjects, TObjectPtr<UObject> InstanceDataObjectsOwner, EViewFlags InViewFlags)
 	: Instances(InstanceDataObjects)
 	, InstancesOwner(InstanceDataObjectsOwner)
-	, RedirectedPropertyTree(MakeShared<FRedirectedPropertyNode>())
 	, ViewFlags(InViewFlags)
 {
-	InitRedirectedPropertyTree();
-}
-
-static bool ObjectHasLoosePropertiesThatNeedFixup(UObject* Object)
-{
-	bool bNeedsFixup = false;
-	Object->GetClass()->Visit(Object, [&bNeedsFixup](const FPropertyVisitorContext& Context)->EPropertyVisitorControlFlow
-	{
-		const FPropertyVisitorPath& Path = Context.Path;
-		const FPropertyVisitorData& Data = Context.Data;
-		const FProperty* Property = Path.Top().Property;
-		if (!Property->HasAnyPropertyFlags(CPF_SkipSerialization) && Property->GetBoolMetaData(NAME_IsLooseMetadata))
-		{
-			bNeedsFixup = true;
-			return EPropertyVisitorControlFlow::Stop;
-		}
-		if (!Property->GetBoolMetaData(NAME_ContainsLoosePropertiesMetadata))
-		{
-			// if this sub-struct doesn't contain loose properties, it won't need fixup
-			return EPropertyVisitorControlFlow::StepOver;
-		}
-		return EPropertyVisitorControlFlow::StepInto;
-	});
-	return bNeedsFixup;
 }
 
 FInstanceDataObjectFixupPanel::~FInstanceDataObjectFixupPanel()
@@ -260,18 +84,16 @@ FInstanceDataObjectFixupPanel::~FInstanceDataObjectFixupPanel()
 
 	ICoreProvider* DataStorage = GetMutableDataStorageFeature<ICoreProvider>(StorageFeatureName);
 	ICompatibilityProvider* DataStorageCompatibility = GetMutableDataStorageFeature<ICompatibilityProvider>(CompatibilityFeatureName);
+	UE::FPropertyBagRepository& Repository = UE::FPropertyBagRepository::Get();
 
 	if (DataStorageCompatibility != nullptr && DataStorage != nullptr)
 	{
 		for (UObject* Instance : Instances)
 		{
-			if (!ObjectHasLoosePropertiesThatNeedFixup(Instance))
+			if (!Repository.RequiresFixup(Repository.FindInstanceForDataObject(Instance)))
 			{
 				static FName AlertName = FName("EntityLoosePropertiesErrorAlert");
 
-				UE::FPropertyBagRepository& Repository = UE::FPropertyBagRepository::Get();
-				Repository.MarkAsFixedUp(Repository.FindInstanceForDataObject(Instance));
-
 				// If a UObject isn't registered with TEDS, there's a chance its parent is registered and is the one
 				// with the alert column on it, so search upward until the nearest registered parent is found.
 				RowHandle Row = DataStorageCompatibility->FindRowWithCompatibleObject(InstancesOwner ? InstancesOwner.Get() : Instance);
@@ -279,17 +101,6 @@ FInstanceDataObjectFixupPanel::~FInstanceDataObjectFixupPanel()
 			}
 		}
 	}
-	else
-	{
-		for (UObject* Instance : Instances)
-		{
-			if (!ObjectHasLoosePropertiesThatNeedFixup(Instance))
-			{
-				UE::FPropertyBagRepository& Repository = UE::FPropertyBagRepository::Get();
-				Repository.MarkAsFixedUp(Repository.FindInstanceForDataObject(Instance));
-			}
-		}
-	}
 }
 
 int32 FInstanceDataObjectFixupPanel::Find(UObject* Value) const
@@ -386,7 +197,7 @@ TSharedPtr<IDetailsView>& FInstanceDataObjectFixupPanel::GenerateDetailsView(boo
 	{
 		DetailsViewArgs.ShouldForceHideProperty.BindLambda([this](const TSharedRef<FPropertyNode>& PropertyNode)->bool
 		{
-			return !IsInRedirectedPropertyTree(*FPropertyNode::CreatePropertyPath(PropertyNode));
+			return !IsPathOverridden(Instances[0], *FPropertyNode::CreatePropertyPath(PropertyNode));
 		});
 	}
 	
@@ -398,6 +209,25 @@ TSharedPtr<IDetailsView>& FInstanceDataObjectFixupPanel::GenerateDetailsView(boo
 	}
 	
 	DetailsView->SetObjects(Instances, true);
+
+	FCoreUObjectDelegates::OnObjectsReinstanced.AddSPLambda(this,
+		[this](const TMap<UObject*, UObject*>& ObjectMap)
+		{
+			bool bFound = false;
+			for (TObjectPtr<UObject>& Ido : Instances)
+			{
+				if (UObject* const* Found = ObjectMap.Find(Ido))
+				{
+					Ido = TObjectPtr<UObject>(*Found);
+					bFound = true;
+				}
+			}
+
+			if (bFound)
+			{
+				DetailsView->SetObjects(Instances, true);
+			}
+		});
 	return DetailsView;
 }
 
@@ -438,9 +268,10 @@ bool FInstanceDataObjectFixupPanel::ShouldSplitterIgnoreRow(const TWeakPtr<FDeta
 
 bool FInstanceDataObjectFixupPanel::AreAllConflictsRedirected() const
 {
+	UE::FPropertyBagRepository& Repository = UE::FPropertyBagRepository::Get();
 	for (UObject* Instance : Instances)
 	{
-		if (ObjectHasLoosePropertiesThatNeedFixup(Instance))
+		if (Repository.RequiresFixup(Repository.FindInstanceForDataObject(Instance)))
 		{
 			return false;
 		}
@@ -467,6 +298,8 @@ void FInstanceDataObjectFixupPanel::AutoApplyMarkDeletedActions()
 					if (Path.IsValid())
 					{
 						MarkForDelete(Path);
+
+						return ETreeTraverseControl::SkipChildren;
 					}
 				}
 			}
@@ -550,384 +383,111 @@ static FPropertyChangedEvent ConstructChangeEventForRedirect(const FPropertyPath
 	return OutEvent;
 }
 
-void FInstanceDataObjectFixupPanel::FTypeConverter::Push(FProperty* SourceProperty, const void* SourceData, FProperty* DestinationProperty, void* DestinationData)
-{
-	InstanceInfo.Push({SourceProperty, SourceData, DestinationProperty, DestinationData});
-	// check if warning was made more severe by this data
-	Warning = FMath::Max(Warning, GenerateWarning(SourceProperty, SourceData, DestinationProperty));
-}
-
-FInstanceDataObjectFixupPanel::FTypeConverter::operator bool() const
-{
-	return Warning != EWarning::InvalidConversion;
-}
-
-void FInstanceDataObjectFixupPanel::FTypeConverter::operator()() const
-{
-	check(Warning != EWarning::InvalidConversion);
-	for (const FInstanceInfo& Info : InstanceInfo)
-	{
-		TryConvert(Info.SourceProperty, Info.SourceData, Info.DestinationProperty, Info.DestinationData);
-	}
-}
-
-FText FInstanceDataObjectFixupPanel::FTypeConverter::GetWarning() const
+static UStruct* ExtractStructFromType(UE::FPropertyTypeName Type, const UE::FPropertyPathName& Path, int32 SegmentIndex)
 {
-	switch(Warning)
+	static const FName NAME_Key(ANSITEXTVIEW("Key"));
+	static const FName NAME_Value(ANSITEXTVIEW("Value"));
+	if (Type.GetName() == NAME_StructProperty)
 	{
-	case EWarning::NarrowingConversion:
-		return LOCTEXT("NarrowingConversion", "This type conversion is a narrowing conversion. Likely data loss!");
-	case EWarning::NonInvertibleConversion:
-		return LOCTEXT("NonInvertibleConversion", "This type conversion is not an invertable operation. Likely data loss!");
-	case EWarning::InvalidConversion:
-		return LOCTEXT("InvalidConversion", "Invalid Conversion");
-	default:
-		return FText::GetEmpty();
-	}
-}
-
-bool FInstanceDataObjectFixupPanel::FTypeConverter::TryConvert(FProperty* SourceProperty, const void* SourceData, FProperty* DestinationProperty, void* DestinationData)
-{
-	
-	FUObjectSerializeContext* SerializeContext = FUObjectThreadContext::Get().GetSerializeContext();
-	TGuardValue<bool> ScopedImpersonateProperties(SerializeContext->bImpersonateProperties, true);
-	TArray<uint8> Buffer;
-	FObjectWriter ObjectWriter(Buffer); // using FObjectWriter so that object properties will serialize correctly
-	FStructuredArchiveFromArchive StructuredWriter(ObjectWriter);
-	// todo: handle static arrays
-	FPropertyTag SourceTag(SourceProperty, 0, (uint8*)SourceData);
-	SourceTag.SerializeTaggedProperty(StructuredWriter.GetSlot(), SourceProperty, (uint8*)SourceData, nullptr);
-	
-	FObjectReader ObjectReader(Buffer);
-	FStructuredArchiveFromArchive StructuredReader(ObjectReader);
-
-	// TODO: this breaks for static array elements.
-	void* DestinationContainer = static_cast<uint8*>(DestinationData) - DestinationProperty->GetOffset_ForInternal();
-
-
-	bool bResult = false;
-	switch(DestinationProperty->ConvertFromType(SourceTag, StructuredReader.GetSlot(), (uint8*)DestinationContainer, SourceProperty->GetOwnerStruct(), nullptr))
-	{
-	case EConvertFromTypeResult::UseSerializeItem:
-		if (SourceProperty->GetID() == DestinationProperty->GetID())
+		UE::FPropertyTypeName StructType = Type.GetParameter(0);
+		FName StructName = StructType.GetName();
+		FName PackageName = StructType.GetParameter(0).GetName();
+		FTopLevelAssetPath StructAssetPath(PackageName, StructName);
+		if (UScriptStruct* Struct = Cast<UScriptStruct>(StaticFindObject(UStruct::StaticClass(), StructAssetPath)))
 		{
-			SourceTag.SerializeTaggedProperty(StructuredReader.GetSlot(), DestinationProperty, (uint8*)DestinationData, nullptr);
-			bResult = true;
+			return Struct;
 		}
-		break;
-	case EConvertFromTypeResult::Serialized:
-		bResult = true;
-		break;
-	case EConvertFromTypeResult::CannotConvert:
-		break;
-	case EConvertFromTypeResult::Converted:
-		bResult = true;
-		break;
 	}
-
-	if (!bResult)
+	else if (Type.GetName() == NAME_ArrayProperty || Type.GetName() == NAME_SetProperty)
 	{
-		bool bTryTextSerialize = false;
-
-		const auto IsStringType = [](const FProperty* Property)
-		{
-			static FName VerseStringName = TEXT("VerseStringProperty");
-			return Property->IsA<FStrProperty>() || Property->IsA<FTextProperty>() || Property->IsA<FNameProperty>() || Property->GetID() == VerseStringName;
-		};
-		
-		if (IsStringType(SourceProperty) || IsStringType(DestinationProperty))
-		{
-			// if either property is a string, text, or name, use text serialization
-			bTryTextSerialize = true;
-		}
-		else if (FStructProperty* SourceAsStructProperty = CastField<FStructProperty>(SourceProperty))
-		{
-			if (FStructProperty* DestinationAsStructProperty = CastField<FStructProperty>(DestinationProperty))
-			{
-				if (!SourceAsStructProperty->Struct->UseNativeSerialization() && !DestinationAsStructProperty->Struct->UseNativeSerialization())
-				{
-					// attempt to text serialize structs since ConvertFromType doesn't support them usually
-					bTryTextSerialize = true;
-				}
-			}
-		}
-
-		// use ExportText_Direct and ImportText_Direct
-		if (bTryTextSerialize)
-		{
-			FString StrBuffer;
-			SourceProperty->ExportText_Direct(StrBuffer, SourceData, nullptr, nullptr, PPF_None);
-			FStringOutputDevice ErrorOutput;
-			DestinationProperty->ImportText_Direct(*StrBuffer, DestinationData, nullptr, PPF_None, &ErrorOutput);
-			bResult = ErrorOutput.IsEmpty();
-		}
-	}
-
-	
-	return bResult;
-}
-
-FInstanceDataObjectFixupPanel::FTypeConverter::EWarning FInstanceDataObjectFixupPanel::FTypeConverter::GenerateWarning(FProperty* SourceProperty, const void* SourceData, FProperty* DestinationProperty)
-{
-	// convert from source to destination in a temp buffer to see if it's possible
-	TArray<uint8, TInlineAllocator<64>> SourceToDest;
-	SourceToDest.SetNumUninitialized(DestinationProperty->GetElementSize());
-	DestinationProperty->InitializeValue(SourceToDest.GetData());
-	if (!TryConvert(SourceProperty, SourceData, DestinationProperty, SourceToDest.GetData()))
-	{
-		return EWarning::InvalidConversion;
-	}
-
-	// convert from destination to source in a temp buffer to see if it's possible
-	TArray<uint8, TInlineAllocator<64>> DestToSource;
-	DestToSource.SetNumUninitialized(SourceProperty->GetElementSize());
-	SourceProperty->InitializeValue(DestToSource.GetData());
-	if (!TryConvert(DestinationProperty, SourceToDest.GetData(), SourceProperty, DestToSource.GetData()))
-	{
-		return EWarning::NonInvertibleConversion;
-	}
-
-	// check that the round trip result has the same value as source
-	if (!SourceProperty->Identical(SourceData, DestToSource.GetData(), PPF_None))
-	{
-		return EWarning::NarrowingConversion;
-	}
-	return EWarning::SafeConversion;
-}
-
-FInstanceDataObjectFixupPanel::FTypeConverter FInstanceDataObjectFixupPanel::CreateTypeConverter(const FPropertyPath& From, const FPropertyPath& To)
-{
-	FTypeConverter Result;
-	for (UObject* Instance : Instances)
-	{
-		FProperty* SourceProperty = From.GetLeafMostProperty().Property.Get();
-		const void* SourceData = ResolvePath(From, Instance);
-		FProperty* DestinationProperty = To.GetLeafMostProperty().Property.Get();
-		void* DestinationData = ResolvePath(To, Instance);
-		Result.Push(SourceProperty, SourceData, DestinationProperty, DestinationData);
+		UE::FPropertyTypeName InnerType = Type.GetParameter(0);
+		return ExtractStructFromType(InnerType, Path, SegmentIndex);
 	}
-	return Result;
-}
-
-void FInstanceDataObjectFixupPanel::RedirectPropertyHelper(const FPropertyPath& From, const FPropertyPath& To, TOptional<FRevertInfo>& FromRevertInfo, FRevertInfo*& ToRevertInfo)
-{
-	UInstanceDataObjectFixupUndoHandler* Snapshot = NewObject<UInstanceDataObjectFixupUndoHandler>();
-	Snapshot->Init(SharedThis(this));
-	GEditor->BeginTransaction(TEXT("InstanceDataObjectFixupTool"), FText::Format(LOCTEXT("RedirectPropertyTransaction","Redirect {0} to {1}"), FText::FromString(From.ToString()), FText::FromString(To.ToString())), nullptr);
-
-	FProperty* SourceProperty = From.GetLeafMostProperty().Property.Get();
-	check(SourceProperty);
-	FProperty* DestinationProperty = To.IsValid() ? To.GetLeafMostProperty().Property.Get() : nullptr;
-	
-	if (const FRevertInfo* Info = RevertInfo.Find(From))
+	else if (Type.GetName() == NAME_MapProperty)
 	{
-		FromRevertInfo = *Info;
-		if (DestinationProperty)
+		if (SegmentIndex + 1 < Path.GetSegmentCount())
 		{
-			if (DestinationProperty->HasAnyPropertyFlags(CPF_SkipSerialization) != Info->bHadSkipSerialization)
+			UE::FPropertyPathNameSegment MapInternalSegment = Path.GetSegment(SegmentIndex + 1);
+			if (MapInternalSegment.Name == NAME_Key)
 			{
-				// toggle CPF_SkipSerialization flag if needed
-				DestinationProperty->PropertyFlags ^= CPF_SkipSerialization;
+				UE::FPropertyTypeName KeyType = Type.GetParameter(0);
+				return ExtractStructFromType(KeyType, Path, SegmentIndex);
 			}
-			if (!Info->bWasHidden)
+			else if (MapInternalSegment.Name == NAME_Value)
 			{
-				DestinationProperty->RemoveMetaData(TEXT("Hidden"));
+				UE::FPropertyTypeName ValType = Type.GetParameter(1);
+				return ExtractStructFromType(ValType.GetParameter(0), Path, SegmentIndex);
 			}
-			DestinationProperty->RemoveMetaData(TEXT("Redirected"));
 		}
-		
-		
-		if (To.IsValid() && To != Info->OriginalPath)
-		{
-			TArray<uint8> OriginalValue;
-			
-			ToRevertInfo = &RevertInfo.Add(To, {
-				.OriginalPath = Info->OriginalPath,
-				.bHadSkipSerialization = SourceProperty->HasAnyPropertyFlags(CPF_SkipSerialization),
-				.bWasHidden = SourceProperty->HasMetaData(TEXT("Hidden"))
-			});
-		}
-		MarkedForDelete.Remove(Info->OriginalPath);
-		RevertInfo.Remove(From);
-	}
-	else
-	{
-		if (To.IsValid())
-		{
-			ToRevertInfo = &RevertInfo.Add(To, {
-				.OriginalPath = From,
-				.bHadSkipSerialization = SourceProperty->HasAnyPropertyFlags(CPF_SkipSerialization)
-			});
-			MarkedForDelete.Remove(From);
-		}
-	}
-	
-	if (To != From)
-	{
-		auto OnHidden = [](FProperty* Property)
-		{
-			if (Property->HasMetaData(NAME_IsLooseMetadata))
-			{
-				Property->PropertyFlags |= CPF_SkipSerialization;
-				Property->SetMetaData(TEXT("Hidden"), TEXT("True"));
-				Property->SetMetaData(TEXT("Redirected"), TEXT("True"));
-			}
-		};
-		if (To.IsValid())
-		{
-			RedirectedPropertyTree->Move(From, To);
-			for (FPropertyPath Path = From; Path.IsValid(); Path = Path.TrimPath(1).Get())
-			{
-				// because RedirectedPropertyTree->Move could've removed multiple properties in the path, we need to check each of them
-				if (!MarkedForDelete.Find(Path) && RedirectedPropertyTree->Find(Path))
-				{
-					break;
-				}
-				OnHidden(Path.GetLeafMostProperty().Property.Get());
-			}
-		}
-		else
-		{
-			OnHidden(SourceProperty);
-		}
-	}
 
-	Snapshot->OnRedirect(From, To);
+	}
+	return nullptr;
 }
 
-void FInstanceDataObjectFixupPanel::RedirectProperty(const FPropertyPath& From, const FPropertyPath& To)
+static int32 GetLastSharedStructSegment(const UE::FPropertyPathName& PathA, const UE::FPropertyPathName& PathB, UStruct*& OutStruct)
 {
-	const FProperty* SourceProperty = From.GetLeafMostProperty().Property.Get();
-	check(SourceProperty);
-	FProperty* DestinationProperty = To.IsValid() ? To.GetLeafMostProperty().Property.Get() : nullptr;
-	
-	TOptional<FRevertInfo> FromRevertInfo;
-	FRevertInfo* ToRevertInfo = nullptr;
-	RedirectPropertyHelper(From, To, FromRevertInfo, ToRevertInfo);
-	
-	if (!DestinationProperty) // null destination is interpreted as a deletion
+	int32 LastStructIndex = INDEX_NONE;
+	const int32 Size = FMath::Min(PathA.GetSegmentCount(), PathB.GetSegmentCount());
+	for (int32 SegmentIndex = 0; SegmentIndex < Size; ++SegmentIndex)
 	{
-		MarkedForDelete.Add(From);
-		GEditor->EndTransaction();
-		DetailsView->ForceRefresh();
-		return; // delete actions don't need data copied
-	}
-	
-	const uint8* FromRevertInfoItr = FromRevertInfo ? FromRevertInfo->OriginalValue.GetData() : nullptr;
-	for (UObject* Instance : Instances)
-	{
-		void* Source = ResolvePath(From, Instance);
-		void* Destination = ResolvePath(To, Instance);
-		
-		if (!ensure(Source && Destination))
+		UE::FPropertyPathNameSegment SegmentA = PathA.GetSegment(SegmentIndex);
+		UE::FPropertyPathNameSegment SegmentB = PathB.GetSegment(SegmentIndex);
+		if (SegmentA.Index != SegmentB.Index || SegmentA.Name != SegmentB.Name || SegmentA.Type != SegmentB.Type)
 		{
-			continue;
+			break;
 		}
-	
 
-		// construct change event
-		FEditPropertyChain Chain;
-		TMap<FString, int32> ArrayIndices;
-		FPropertyChangedEvent ChangeEvent = ConstructChangeEventForRedirect(To, Chain, ArrayIndices);
-		FPropertyChangedChainEvent ChangedChainEvent(Chain, ChangeEvent);
-		Instance->PreEditChange(Chain);
-
-		if (ToRevertInfo)
-		{
-			// cache the destination value so it can be reverted later
-			const int32 Size = DestinationProperty->ArrayDim * DestinationProperty->GetElementSize();
-			ToRevertInfo->OriginalValue.AddZeroed(Size);
-			uint8* Buffer = ToRevertInfo->OriginalValue.GetData() + (ToRevertInfo->OriginalValue.Num() - Size);
-			DestinationProperty->CopyCompleteValue(Buffer, Destination);
-		}
-		
-		if (SourceProperty->SameType(DestinationProperty))
+		if (UStruct* Found = ExtractStructFromType(SegmentA.Type, PathA, SegmentIndex))
 		{
-			SourceProperty->CopyCompleteValue(Destination, Source);
-			FOverridableManager::Get().GetOverriddenProperties(Instance)->RestoreOverriddenPropertyOperation(EOverriddenPropertyOperation::Modified, nullptr, DestinationProperty);
+			OutStruct = Found;
+			LastStructIndex = SegmentIndex;
 		}
-		else
-		{
-			FString ValueStr;
-			SourceProperty->ExportText_Direct(ValueStr, Source, nullptr, Instance, PPF_Copy);
-			DestinationProperty->ImportText_Direct(*ValueStr, Destination, Instance, PPF_Copy);
-		}
-		
-		if (FromRevertInfo)
-		{
-			// apply FromRevertInfo to From
-			SourceProperty->CopyCompleteValue(Source, FromRevertInfoItr);
-			FromRevertInfoItr += DestinationProperty->ArrayDim * DestinationProperty->GetElementSize();
-		}
-		Instance->PostEditChangeChainProperty(ChangedChainEvent);
 	}
-
-	GEditor->EndTransaction();
-	DetailsView->ForceRefresh();
+	return LastStructIndex;
 }
 
-void FInstanceDataObjectFixupPanel::RedirectProperty(const FPropertyPath& From, const FPropertyPath& To, const FTypeConverter& TypeConversion)
+void FInstanceDataObjectFixupPanel::RedirectProperty(const FPropertyPath& From, const FPropertyPath& To)
 {
-	const FProperty* SourceProperty = From.GetLeafMostProperty().Property.Get();
-	const FProperty* DestinationProperty = To.GetLeafMostProperty().Property.Get();
-	check(SourceProperty && DestinationProperty);
-	
-	TOptional<FRevertInfo> FromRevertInfo;
-	FRevertInfo* ToRevertInfo = nullptr;
-	RedirectPropertyHelper(From, To, FromRevertInfo, ToRevertInfo);
-	
-	const uint8* FromRevertInfoItr = FromRevertInfo ? FromRevertInfo->OriginalValue.GetData() : nullptr;
-	TArray<FPropertyChangedEvent> ChangeEvents;
-	TArray<FEditPropertyChain> Chains;
+	FUObjectSerializeContext* SerializeContext = FUObjectThreadContext::Get().GetSerializeContext();
+	TGuardValue<bool> ScopedImpersonateProperties(SerializeContext->bImpersonateProperties, true);
+	TArray<UE::FPropertyPathName> FromPathNames = From.ToPropertyPathName();
+	TArray<UE::FPropertyPathName> ToPathNames = To.ToPropertyPathName();
 
-	// call PreEditChange and set up undo handling
-	for (UObject* Instance : Instances)
+	if (FromPathNames.Num() != 1 || ToPathNames.Num() != 1)
 	{
-		void* Source = ResolvePath(From, Instance);
-		void* Destination = ResolvePath(To, Instance);
-		
-		if (!ensure(Source && Destination))
-		{
-			continue;
-		}
-		
-		// construct change event
-		Chains.Emplace();
-		TMap<FString, int32> ArrayIndices;
-		ChangeEvents.Add(ConstructChangeEventForRedirect(To, Chains.Last(), ArrayIndices));
-		Instance->PreEditChange(Chains.Last());
-
-		if (ToRevertInfo)
-		{
-			// cache the destination value so it can be reverted later
-			const int32 Size = DestinationProperty->ArrayDim * DestinationProperty->GetElementSize();
-			ToRevertInfo->OriginalValue.AddZeroed(Size);
-			uint8* Buffer = ToRevertInfo->OriginalValue.GetData() + (ToRevertInfo->OriginalValue.Num() - Size);
-			DestinationProperty->CopyCompleteValue(Buffer, Destination);
-		}
+		ensureMsgf(false, TEXT("Redirecting into or out of subobjects is not supported yet!"));
+		return;
 	}
 
-	// applied to all instances at once
-	TypeConversion();
+	// So that this IDT operation can be used on as many struct instances as possible, find the inner-most struct that's shared by both paths and chop them to be relative
+	// to that struct. So for example 'From' is "Foo (MyFooStruct) -> Bar (MyBarStruct) -> X (Float)" and 'To' is "Foo (MyFooStruct) -> Bar (MyBarStruct) -> Y (Float)"
+	// then we'll create an IDT for ALL instances of MyBarStruct that remaps "X (Float)" to "Y (Float)" instead of only remapping the full path.
+	const UObject* Owner = UE::FPropertyBagRepository::Get().FindInstanceForDataObject(Instances[0]);
+	UStruct* BaseStruct = Owner->GetClass();
+	int32 StructIndex = GetLastSharedStructSegment(FromPathNames[0], ToPathNames[0], BaseStruct);
 
-	// call post edit change and apply undo handling
-	for (int32 I = 0; I < Instances.Num(); ++I)
+	UE::FPropertyPathName SrcPath;
+	UE::FPropertyPathName DstPath;
+	for (int32 SegmentIndex = StructIndex + 1; SegmentIndex < FromPathNames[0].GetSegmentCount() || SegmentIndex < ToPathNames[0].GetSegmentCount(); ++SegmentIndex)
 	{
-		UObject* Instance = Instances[I];
-		void* Source = ResolvePath(From, Instance);
-		if (FromRevertInfo)
+		if (SegmentIndex < ToPathNames[0].GetSegmentCount())
+		{
+			DstPath.Push(ToPathNames[0].GetSegment(SegmentIndex));
+		}
+		if (SegmentIndex < FromPathNames[0].GetSegmentCount())
 		{
-			// apply FromRevertInfo to From
-			SourceProperty->CopyCompleteValue(Source, FromRevertInfoItr);
-			FromRevertInfoItr += DestinationProperty->ArrayDim * DestinationProperty->GetElementSize();
+			SrcPath.Push(FromPathNames[0].GetSegment(SegmentIndex));
 		}
-		FPropertyChangedChainEvent ChangedChainEvent(Chains[I], ChangeEvents[I]);
-		Instance->PostEditChangeChainProperty(ChangedChainEvent);
 	}
 
-	GEditor->EndTransaction();
-	
-	DetailsView->ForceRefresh();
+	TSharedPtr<FJsonObject> Params = MakeShared<FJsonObject>();
+	Params->SetField(TEXT("src"), SrcPath.ToJsonValue());
+	Params->SetField(TEXT("dst"), DstPath.ToJsonValue());
+	UE::FInstanceDataTransforms::Get().WriteInstanceDataTransform(*BaseStruct, TEXT("RedirectProperty"), Params);
+	TObjectPtr<UObject> OwnerObjectPtr = const_cast<UObject*>(Owner);
+
+	// this will likely trigger a IDO reinstance which will cause a details panel refresh for the new object!
+	UE::FInstanceDataTransforms::Get().PatchInstanceDataObject(Instances[0], OwnerObjectPtr);
 }
 
 void FInstanceDataObjectFixupPanel::OnRedirectProperty(FPropertyPath From, FPropertyPath To)
@@ -935,141 +495,28 @@ void FInstanceDataObjectFixupPanel::OnRedirectProperty(FPropertyPath From, FProp
 	RedirectProperty(From, To);
 }
 
-void FInstanceDataObjectFixupPanel::OnRedirectProperty(FPropertyPath From, FPropertyPath To, FTypeConverter TypeConversion)
+bool FInstanceDataObjectFixupPanel::PropertyNeedsRedirect(const FPropertyPath& Path) const
 {
-	RedirectProperty(From, To, TypeConversion);
-}
-
-static void InitRedirectedPropertyTreeRec(const TSharedPtr<FRedirectedPropertyNode>& Node, FProperty* Property, void* Value, TSet<UObject*>& EnteredObjects);
-static void InitRedirectedPropertyTreeRec(const TSharedPtr<FRedirectedPropertyNode>& Node, UStruct* Struct, void* StructValue, TSet<UObject*>& EnteredObjects)
-{
-	const UE::FSerializedPropertyValueState SerializedState(Struct, StructValue);
-	for (FProperty* Property : TFieldRange<FProperty>(Struct))
+	if (Path.GetLeafMostProperty().Property->HasAnyPropertyFlags(CPF_SkipSerialization))
 	{
-		for (int32 StaticArrayIndex = 0; StaticArrayIndex < Property->ArrayDim; ++StaticArrayIndex)
-		{
-			if (SerializedState.IsSet(Property, StaticArrayIndex))
-			{
-				FPropertyInfo PropertyInfo(
-					Property, 
-					Property->ArrayDim > 1 ? StaticArrayIndex : INDEX_NONE // use INDEX_NONE for properties that aren't in arrays
-				);
-				const TSharedPtr<FRedirectedPropertyNode>& ChildNode = Node->FindOrAdd(PropertyInfo);
-				void* Value = Property->ContainerPtrToValuePtr<void>(StructValue, StaticArrayIndex);
-				InitRedirectedPropertyTreeRec(ChildNode, Property, Value, EnteredObjects);
-			}
-		}
+		return false;
 	}
-}
 
-static void InitRedirectedPropertyTreeRec(const TSharedPtr<FRedirectedPropertyNode>& Node, FProperty* Property, void* Value, TSet<UObject*>& EnteredObjects)
-{
-	if (const FStructProperty* AsStructProperty = CastField<FStructProperty>(Property))
-	{
-		InitRedirectedPropertyTreeRec(Node, AsStructProperty->Struct, Value, EnteredObjects);
-	}
-	else if (const FObjectProperty* AsObjectProperty = CastField<FObjectProperty>(Property))
-	{
-		if (AsObjectProperty->HasAnyPropertyFlags(CPF_InstancedReference))
-		{
-			if (UObject* Object = AsObjectProperty->GetObjectPropertyValue(Value))
-            {
-				UE::FPropertyBagRepository& PropertyBagRepository = UE::FPropertyBagRepository::Get();
-				if (UObject* Found = PropertyBagRepository.FindInstanceDataObject(Object))
-				{
-					Object = Found;
-				}
-				// check for circular references to avoid infinite recursion
-				if (!EnteredObjects.Contains(Object))
-				{
-					EnteredObjects.Add(Object);
-            		InitRedirectedPropertyTreeRec(Node, Object->GetClass(), Object, EnteredObjects);
-					EnteredObjects.Remove(Object);
-            	}
-            }
-		}
-		
-	}
-	else if (const FArrayProperty* AsArrayProperty = CastField<FArrayProperty>(Property))
+	if (!IsPathOverridden(Instances[0], Path))
 	{
-		FScriptArrayHelper Array(AsArrayProperty, Value);
-		for (int32 ArrayIndex = 0; ArrayIndex < Array.Num(); ++ArrayIndex)
-		{
-			const TSharedPtr<FRedirectedPropertyNode>& ChildNode = Node->FindOrAdd(FPropertyInfo(AsArrayProperty->Inner, ArrayIndex));
-			InitRedirectedPropertyTreeRec(ChildNode, AsArrayProperty->Inner, Array.GetElementPtr(ArrayIndex), EnteredObjects);
-		}
+		return false;
 	}
-	else if (const FSetProperty* AsSetProperty = CastField<FSetProperty>(Property))
-	{
-		FScriptSetHelper Set(AsSetProperty, Value);
-		for (FScriptSetHelper::FIterator Itr = Set.CreateIterator(); Itr; ++Itr)
-		{
-			const TSharedPtr<FRedirectedPropertyNode>& ChildNode = Node->FindOrAdd(FPropertyInfo(AsSetProperty->ElementProp, Itr.GetLogicalIndex()));
-			InitRedirectedPropertyTreeRec(ChildNode, AsSetProperty->ElementProp, Set.GetElementPtr(Itr), EnteredObjects);
-		}
-	}
-	else if (const FMapProperty* AsMapProperty = CastField<FMapProperty>(Property))
-	{
-		FScriptMapHelper Map(AsMapProperty, Value);
-		for (FScriptMapHelper::FIterator Itr = Map.CreateIterator(); Itr; ++Itr)
-		{
-			const TSharedPtr<FRedirectedPropertyNode>& KeyNode = Node->FindOrAdd(FPropertyInfo(AsMapProperty->KeyProp, Itr.GetLogicalIndex()));
-			InitRedirectedPropertyTreeRec(KeyNode, AsMapProperty->KeyProp, Map.GetKeyPtr(Itr), EnteredObjects);
-			const TSharedPtr<FRedirectedPropertyNode>& ValNode = Node->FindOrAdd(FPropertyInfo(AsMapProperty->ValueProp, Itr.GetLogicalIndex()));
-			InitRedirectedPropertyTreeRec(ValNode, AsMapProperty->ValueProp, Map.GetValuePtr(Itr), EnteredObjects);
-		}
-	}
-}
 
-void FInstanceDataObjectFixupPanel::InitRedirectedPropertyTree()
-{
-	TSet<UObject*> EnteredObjects = {Instances[0]};
-	InitRedirectedPropertyTreeRec(RedirectedPropertyTree, Instances[0]->GetClass(), Instances[0], EnteredObjects);
-	EnteredObjects.Remove(Instances[0]);
-	check(EnteredObjects.IsEmpty());
-}
-
-void UInstanceDataObjectFixupUndoHandler::Init(const TSharedRef<FInstanceDataObjectFixupPanel>& Panel)
-{
-	InstanceDataObjectPanel = Panel;
-	RevertInfo = Panel->RevertInfo;
-	MarkedForDelete = Panel->MarkedForDelete;
-	SetFlags(RF_Transactional);
-}
-
-void UInstanceDataObjectFixupUndoHandler::OnRedirect(const FPropertyPath& From, const FPropertyPath& To)
-{
-	if (const TSharedPtr<FInstanceDataObjectFixupPanel> Panel = InstanceDataObjectPanel.Pin())
-	{
-		RedirectFrom = From;
-		RedirectTo = To;
-		++ChangeNum;
-	}
-	Modify();
-}
-
-void UInstanceDataObjectFixupUndoHandler::PostEditUndo()
-{
-	if (const TSharedPtr<FInstanceDataObjectFixupPanel> Panel = InstanceDataObjectPanel.Pin())
+	FUObjectSerializeContext* SerializeContext = FUObjectThreadContext::Get().GetSerializeContext();
+	TGuardValue<bool> ScopedImpersonateProperties(SerializeContext->bImpersonateProperties, true);
+	TArray<UE::FPropertyPathName> PathNames = Path.ToPropertyPathName();
+	if (PathNames.Num() != 1)
 	{
-		if (RedirectTo != RedirectFrom)
-		{
-			if (RedirectTo.IsValid() && RedirectFrom.IsValid())
-			{
-				Panel->RedirectedPropertyTree->Move(RedirectTo, RedirectFrom);
-			}
-			Swap(RedirectTo, RedirectFrom);
-		}
-		
-		Swap(Panel->RevertInfo, RevertInfo);
-		Swap(Panel->MarkedForDelete, MarkedForDelete);
-		Panel->DetailsView->ForceRefresh();
+		return false;
 	}
-}
 
-bool FInstanceDataObjectFixupPanel::IsInRedirectedPropertyTree(const FPropertyPath& Path) const
-{
-	return RedirectedPropertyTree->Find(Path).IsValid();
+	const UObject* Owner = UE::FPropertyBagRepository::Get().FindInstanceForDataObject(Instances[0]);
+	return UE::FPropertyBagRepository::Get().RequiresFixup(Owner, PathNames[0]);
 }
 
 const FPropertyPath& FInstanceDataObjectFixupPanel::GetOriginalPath(const FPropertyPath& Path) const
@@ -1083,18 +530,21 @@ const FPropertyPath& FInstanceDataObjectFixupPanel::GetOriginalPath(const FPrope
 
 void FInstanceDataObjectFixupPanel::MarkForDelete(const FPropertyPath& CurrentPath)
 {
-	// undo any existing redirection on this node
-	if (const FRevertInfo* Found = RevertInfo.Find(CurrentPath))
-	{
-		// move this property back to it's original location before marking it for delete
-		const FPropertyPath PathCopy = Found->OriginalPath; // RedirectProperty will invalidate pointers. copy path by value so it doesn't get destroyed.
-		RedirectProperty(CurrentPath, PathCopy);
-		RedirectProperty(PathCopy, {});
-	}
-	else
+	FUObjectSerializeContext* SerializeContext = FUObjectThreadContext::Get().GetSerializeContext();
+	TGuardValue<bool> ScopedImpersonateProperties(SerializeContext->bImpersonateProperties, true);
+	TArray<UE::FPropertyPathName> srcPathNames = CurrentPath.ToPropertyPathName();
+
+	if (srcPathNames.Num() != 1)
 	{
-		RedirectProperty(CurrentPath, {});
+		ensureMsgf(false, TEXT("out of subobjects is not supported yet!"));
+		return;
 	}
+	const UObject* Owner = UE::FPropertyBagRepository::Get().FindInstanceForDataObject(Instances[0]);
+	TSharedPtr<FJsonObject> Params = MakeShared<FJsonObject>();
+	Params->SetField(TEXT("src"), srcPathNames[0].ToJsonValue());
+	UE::FInstanceDataTransforms::Get().WriteInstanceDataTransform(*Owner->GetClass(), TEXT("RemoveProperty"), Params);
+	TObjectPtr<UObject> OwnerObjectPtr = const_cast<UObject*>(Owner);
+	UE::FInstanceDataTransforms::Get().PatchInstanceDataObject(Instances[0], OwnerObjectPtr);
 }
 
 void FInstanceDataObjectFixupPanel::OnMarkForDelete(FPropertyPath Path)
diff --git a/Engine/Plugins/Experimental/InstanceDataObjectFixupTool/Source/Private/InstanceDataObjectFixupPanel.h b/Engine/Plugins/Experimental/InstanceDataObjectFixupTool/Source/Private/InstanceDataObjectFixupPanel.h
--- a/Engine/Plugins/Experimental/InstanceDataObjectFixupTool/Source/Private/InstanceDataObjectFixupPanel.h
+++ b/Engine/Plugins/Experimental/InstanceDataObjectFixupTool/Source/Private/InstanceDataObjectFixupPanel.h
@@ -6,36 +6,9 @@
 #include "AsyncDetailViewDiff.h"
 #include "UObject/Object.h"
 #include "Templates/SharedPointer.h"
-#include "InstanceDataObjectFixupPanel.generated.h"
 
 class SLinkableScrollBar;
 
-struct FRedirectedPropertyNode : public TSharedFromThis<FRedirectedPropertyNode>
-{
-	FRedirectedPropertyNode() = default;
-	FRedirectedPropertyNode(const FRedirectedPropertyNode& Other);
-	FRedirectedPropertyNode(const FPropertyInfo& Info, const TWeakPtr<FRedirectedPropertyNode>& Parent);
-	FRedirectedPropertyNode(FName InPropertyName, const UE::FPropertyTypeName& InType, int32 InArrayIndex, const TWeakPtr<FRedirectedPropertyNode>& InParent);
-	TSharedPtr<FRedirectedPropertyNode> FindOrAdd(const FPropertyPath& Path, int32 PathIndex = 0);
-	TSharedPtr<FRedirectedPropertyNode> FindOrAdd(const FPropertyInfo& ChildInfo);
-	TSharedPtr<FRedirectedPropertyNode> FindOrAdd(FName ChildPropertyName, const UE::FPropertyTypeName& ChildType, int32 ChildArrayIndex = 0);
-	bool Remove(const FPropertyPath& Path, int32 PathIndex = 0);
-	bool Remove(const FPropertyInfo& ChildInfo);
-	bool Remove(FName ChildPropertyName, const UE::FPropertyTypeName& ChildType, int32 ChildArrayIndex = 0);
-	TSharedPtr<FRedirectedPropertyNode> Find(const FPropertyPath& Path, int32 PathIndex = 0) const;
-	TSharedPtr<FRedirectedPropertyNode> Find(const FPropertyInfo& ChildInfo) const;
-	TSharedPtr<FRedirectedPropertyNode> Find(FName ChildPropertyName, const UE::FPropertyTypeName& ChildType, int32 ChildArrayIndex = 0) const;
-	bool Move(const FPropertyPath& FromPath, const FPropertyPath& ToPath);
-	int32 FindIndex(const FPropertyInfo& ChildInfo) const;
-	int32 FindIndex(FName ChildPropertyName, const UE::FPropertyTypeName& ChildType, int32 ChildArrayIndex = 0) const;
-
-	FName PropertyName;
-	UE::FPropertyTypeName Type;
-	int32 ArrayIndex;
-	TWeakPtr<FRedirectedPropertyNode> Parent;
-	TArray<TSharedPtr<FRedirectedPropertyNode>> Children;
-};
-
 class FInstanceDataObjectFixupPanel : public TSharedFromThis<FInstanceDataObjectFixupPanel>
 {
 public:
@@ -75,9 +48,8 @@ public:
 	// mark all conflicted properties for delete (FixupMode only)
 	void AutoApplyMarkDeletedActions();
 
-	// the redirected property tree keeps track of which properties in the InstanceDataObject were either set by a property bag during serialization or
-	// redirected to from a floating property. The members of the tree are visible in the left panel.
-	bool IsInRedirectedPropertyTree(const FPropertyPath& Path) const;
+	// returns true if the provided path is an unknown property that needs to be redirected via IDT
+	bool PropertyNeedsRedirect(const FPropertyPath& Path) const;
 
 	bool HasViewFlag(EViewFlags Flag);
 
@@ -85,50 +57,6 @@ public:
 	TSharedPtr<IDetailsView> DetailsView;
 	TSharedPtr<SLinkableScrollBar> LinkableScrollBar;
 
-	// functor used to generate warnings about a type conversion (or set of conversions) and invoke that conversion on demand
-	struct FTypeConverter
-	{
-		FTypeConverter() = default;
-		
-		// add a conversion that should be performed call operator is invoked.
-		void Push(FProperty* SourceProperty, const void* SourceData, FProperty* DestinationProperty, void* DestinationData);
-
-		// return whether this is a valid conversion. (only true if all conversions pushed are valid)
-		operator bool() const;
-
-		// run conversion on all pushed data
-		void operator()() const;
-
-		// return the most severe text warning for the conversions pushed.
-		FText GetWarning() const;
-	private:
-		enum class EWarning
-		{
-			// sorted by least severe to most severe.
-			SafeConversion,
-			NarrowingConversion,
-			NonInvertibleConversion,
-			InvalidConversion,
-		};
-		struct FInstanceInfo
-		{
-			FProperty* SourceProperty;
-			const void* SourceData;
-			FProperty* DestinationProperty;
-			void* DestinationData;
-		};
-		
-		static bool TryConvert(FProperty* SourceProperty, const void* SourceData, FProperty* DestinationProperty, void* DestinationData);
-		static EWarning GenerateWarning(FProperty* SourceProperty, const void* SourceData, FProperty* DestinationProperty);
-
-		// most severe conversion warning found in all the pushed data
-		EWarning Warning = EWarning::SafeConversion;
-
-		TArray<FInstanceInfo> InstanceInfo;
-	};
-	
-	FTypeConverter CreateTypeConverter(const FPropertyPath& From, const FPropertyPath& To);
-
 private:
 	friend class FInstanceDataObjectFixupSpecification; // for access to Redirects
 	friend class FInstanceDataObjectNameWidgetOverride;
@@ -141,23 +69,13 @@ private:
 		bool bWasHidden;
 	};
 
-	void RedirectPropertyHelper(const FPropertyPath& From, const FPropertyPath& To, TOptional<FRevertInfo>& FromRevertInfo, FRevertInfo*& ToRevertInfo);
 	void RedirectProperty(const FPropertyPath& From, const FPropertyPath& To);
-	void RedirectProperty(const FPropertyPath& From, const FPropertyPath& To, const FTypeConverter& TypeConversion);
 	// pass-by-copy version for delegates. Use RedirectProperty when possible
 	void OnRedirectProperty(FPropertyPath From, FPropertyPath To);
-	void OnRedirectProperty(FPropertyPath From, FPropertyPath To, FTypeConverter TypeConversion);
-	
-	void InitRedirectedPropertyTree();
 	
 	TArray<TObjectPtr<UObject>> Instances; // stores either InstanceDataObject property bag or just regular objects
 	TObjectPtr<UObject> InstancesOwner; // if not null, then it stores the common parent for the instances.
 
-	// the redirected property tree keeps track of which properties in the InstanceDataObject were either set by a property bag during serialization or
-	// redirected to from a floating property. The members of the tree are visible in the left panel.
-	TSharedPtr<FRedirectedPropertyNode> RedirectedPropertyTree;
-
-	
 	TMap<FPropertyPath, FRevertInfo> RevertInfo;
 	TSet<FPropertyPath> MarkedForDelete;
 	
@@ -165,23 +83,3 @@ private:
 	TWeakPtr<FAsyncDetailViewDiff> DiffAgainstRight;
 	EViewFlags ViewFlags = EViewFlags::None;
 };
-
-UCLASS()
-class UInstanceDataObjectFixupUndoHandler : public UObject
-{
-public:
-	GENERATED_BODY()
-	void Init(const TSharedRef<FInstanceDataObjectFixupPanel>& Panel);
-	void OnRedirect(const FPropertyPath& From, const FPropertyPath& To);
-	virtual void PostEditUndo() override;
-	
-	TWeakPtr<FInstanceDataObjectFixupPanel> InstanceDataObjectPanel;
-	
-	TMap<FPropertyPath, FInstanceDataObjectFixupPanel::FRevertInfo> RevertInfo;
-	TSet<FPropertyPath> MarkedForDelete;
-	FPropertyPath RedirectFrom;
-	FPropertyPath RedirectTo;
-	
-	UPROPERTY()
-	int32 ChangeNum = 0;
-};
diff --git a/Engine/Source/Runtime/CoreUObject/Internal/UObject/InstanceDataTransforms.h b/Engine/Source/Runtime/CoreUObject/Internal/UObject/InstanceDataTransforms.h
--- a/Engine/Source/Runtime/CoreUObject/Internal/UObject/InstanceDataTransforms.h
+++ b/Engine/Source/Runtime/CoreUObject/Internal/UObject/InstanceDataTransforms.h
@@ -5,15 +5,65 @@
 #include "Containers/Map.h"
 #include "Templates/Function.h"
 #include "Templates/SharedPointer.h"
+#include "Misc/DateTime.h"
+#include "Misc/NotNull.h"
+#include "UObject/ObjectPtr.h"
+#include "UObject/Object.h"
 
 #if !UE_BUILD_SHIPPING_WITH_EDITOR && WITH_EDITORONLY_DATA 
 
-class UObject;
 class FJsonObject;
+struct FTopLevelAssetPath;
+namespace UE { class FPropertyPathName; }
+namespace UE::Private { struct FRegesterInstanceDataTransformHelper; }
+
+namespace UE::Private::IDT
+{
+	using FTransformParams = TSharedPtr<FJsonObject>;
+	enum ETransformResult
+	{
+		Break,
+		Continue
+	};
+	using FTransformFunction = TFunction<ETransformResult(TObjectPtr<UObject>& IDO, TObjectPtr<UObject>& Owner, const FPropertyPathName& BasePath, const FTransformParams&)>;
+}
 
 namespace UE
 {
-	COREUOBJECT_API bool PatchInstanceDataObject(TObjectPtr<UObject>& IDO, TObjectPtr<UObject>& Owner);
+	class FInstanceDataTransforms
+	{
+	public:
+		using FTransformParams = UE::Private::IDT::FTransformParams;
+		using ETransformResult = UE::Private::IDT::ETransformResult;
+		using FTransformFunction = UE::Private::IDT::FTransformFunction;
+
+		COREUOBJECT_API static FInstanceDataTransforms& Get();
+
+		COREUOBJECT_API void WriteInstanceDataTransform(UStruct& ClassOrStruct, const FString& OpCode, const FTransformParams& Params);
+		COREUOBJECT_API TObjectPtr<UObject> PatchInstanceDataObject(TNotNull<UObject*> IDO, TNotNull<UObject*> Owner);
+
+	private:
+		struct FTransform
+		{
+			FTransformParams Args;
+			FTransformFunction Func;
+		};
+
+		const TArray<TSharedPtr<FInstanceDataTransforms::FTransform>>* FindTransformsInternal(const FTopLevelAssetPath& ClassPath);
+		const TArray<TSharedPtr<FInstanceDataTransforms::FTransform>>* FindTransformsInternal(const TNotNull<UStruct*> Struct);
+
+		void ReadConfigFile(const FString& ConfigFilepath);
+
+		friend Private::FRegesterInstanceDataTransformHelper;
+		void RegisterOperation(FStringView Name, FTransformFunction Function);
+
+		void PatchClass(TObjectPtr<UObject>& IDO, TObjectPtr<UObject>& Owner);
+		void PatchStructs(TObjectPtr<UObject>& IDO, TObjectPtr<UObject>& Owner);
+
+		TMap<FString, FDateTime> ConfigFileLastReadTime;
+		TMap<FTopLevelAssetPath, TArray<TSharedPtr<FTransform>>> StructPathToIDTs;
+		TMap<FString, FTransformFunction> RegisteredOperations;
+	};
 }
 
 #endif // !UE_BUILD_SHIPPING_WITH_EDITOR && WITH_EDITORONLY_DATA 
\ No newline at end of file
diff --git a/Engine/Source/Runtime/CoreUObject/Internal/UObject/PropertyBagRepository.h b/Engine/Source/Runtime/CoreUObject/Internal/UObject/PropertyBagRepository.h
--- a/Engine/Source/Runtime/CoreUObject/Internal/UObject/PropertyBagRepository.h
+++ b/Engine/Source/Runtime/CoreUObject/Internal/UObject/PropertyBagRepository.h
@@ -14,6 +14,7 @@ struct FPropertyChangedChainEvent;
 
 namespace UE
 {
+	class FPropertyPathName;
 
 #if WITH_EDITORONLY_DATA
 
@@ -35,7 +36,6 @@ class FPropertyBagRepository
 		void Destroy();
 
 		TObjectPtr<UObject> InstanceDataObject = nullptr;
-		bool bNeedsFixup = false;
 	};
 	// TODO: Make private throughout and extend access permissions here or in wrapper classes? Don't want engine code modifying bags outside of serializers and details panels.
 	//friend UObjectBase;
@@ -109,8 +109,13 @@ public:
 	 * @return				- Does the object's InstanceDataObject, or optionally any of its outer objects, contain any loose properties requiring user fixup before the object may be published?
 	 */
 	COREUOBJECT_API bool RequiresFixup(const UObject* Object, bool bIncludeOuter = false) const;
-	// set the bNeedsFixup flag for this object's IDO to false
-	COREUOBJECT_API void MarkAsFixedUp(const UObject* Object = nullptr);
+
+	/**
+	 * Return true if the object has an unknown property at the provided path
+	 * @param Object - object to test
+	 * @param Path - the specific property path to look up and check for an unknown property
+	 */
+	COREUOBJECT_API bool RequiresFixup(const UObject* Object, const FPropertyPathName& Path) const;
 
 	// Accessors
 	COREUOBJECT_API bool HasInstanceDataObject(const UObject* Owner) const;
diff --git a/Engine/Source/Runtime/CoreUObject/Private/UObject/InstanceDataTransforms.cpp b/Engine/Source/Runtime/CoreUObject/Private/UObject/InstanceDataTransforms.cpp
--- a/Engine/Source/Runtime/CoreUObject/Private/UObject/InstanceDataTransforms.cpp
+++ b/Engine/Source/Runtime/CoreUObject/Private/UObject/InstanceDataTransforms.cpp
@@ -14,14 +14,15 @@
 #include "UObject/StrProperty.h"
 #include "UObject/TextProperty.h"
 #include "UObject/PropertyBagRepository.h"
+#include "UObject/OverrideEventHelper.h"
 #include "UObject/PropertyHelper.h"
 #include "UObject/Package.h"
 #include "Dom/JsonValue.h"
 #include "Dom/JsonObject.h"
 #include "Serialization/JsonReader.h"
+#include "Serialization/JsonWriter.h"
 #include "Serialization/JsonSerializer.h"
 #include "HAL/FileManager.h"
-#include "Misc/DateTime.h"
 #include "Misc/FileHelper.h"
 #include "Misc/Paths.h"
 #include "Misc/PackageName.h"
@@ -30,185 +31,37 @@
 
 DEFINE_LOG_CATEGORY_STATIC(LogInstanceDataTransforms, Display, Display);
 
-// registers a method that can be looked up by a string name in FInstanceDataTransformRepository::RegisteredOperations
+// registers a method that can be looked up by a string name in FInstanceDataTransforms::RegisteredOperations
 #define REGISTER_IDT_OPERATION(FunctionName) \
 UE::Private::FRegesterInstanceDataTransformHelper _Registered##FunctionName{TEXT(#FunctionName), FunctionName}; \
 
 
 namespace UE::Private
 {
-	enum EInstanceDataTransformResult
-	{
-		Break,
-		Continue
-	};
-
-	typedef TSharedPtr<FJsonObject> FInstanceDataTransformParams;
-	typedef TFunction<EInstanceDataTransformResult(TObjectPtr<UObject>& IDO, TObjectPtr<UObject>& Owner, const FPropertyPathName& BasePath, const FInstanceDataTransformParams&)> FInstanceDataTransformFunction;
 
-	struct FInstanceDataTransform
+	FString InstanceDataTransformDirectory(const FTopLevelAssetPath& ClassOrStructPath)
 	{
-		FInstanceDataTransformParams Args;
-		FInstanceDataTransformFunction Func;
-	};
-
-	class FInstanceDataTransformRepository
-	{
-	public:
-		static FInstanceDataTransformRepository& Get()
-		{
-			static FInstanceDataTransformRepository Instance;
-			return Instance;
-		}
-
-		const TArray<TSharedPtr<FInstanceDataTransform>>* FindTransforms(const FTopLevelAssetPath& ClassPath)
+		FPackagePath PackagePath = FPackagePath::FromPackageNameChecked(ClassOrStructPath.GetPackageName());
+		if (PackagePath.IsMountedPath())
 		{
-			// the transform config for a class is found in <RootPackageDirectory>/Config/InstanceDataTransforms.json
-			// TODO: I need to some research on wheather any work needs to be done in regards to making sure this file is treated properly by source control and the content worker.
-			FPackagePath PackagePath = FPackagePath::FromPackageNameChecked(ClassPath.GetPackageName());
-			if (PackagePath.IsMountedPath())
+			PackagePath.SetHeaderExtension(EPackageExtension::Asset); // silence missing extension warnings
+			FString Directory = FPaths::GetPath(PackagePath.GetLocalFullPath());
+			if (int32 Found = Directory.Find(TEXT("Content")); Found != INDEX_NONE)
 			{
-				PackagePath.SetHeaderExtension(EPackageExtension::Asset); // silence missing extension warnings
-				FString Directory = FPaths::GetPath(PackagePath.GetLocalFullPath());
-				if (int32 Found = Directory.Find(TEXT("Content")); Found != INDEX_NONE)
-				{
-					Directory.LeftInline(Found);
-				}
-				Directory /= TEXT("Config");
-
-				FString ConfigFilepath = Directory / TEXT("InstanceDataTransforms.json");
-				if (IFileManager::Get().FileExists(*ConfigFilepath))
-				{
-					FDateTime FileModificationTime = IFileManager::Get().GetTimeStamp(*ConfigFilepath);
-
-					if (FDateTime* LastReadTime = ConfigFileLastReadTime.Find(ConfigFilepath))
-					{
-						// check if the config file has been updated since the last time we read it
-						if (*LastReadTime < FileModificationTime)
-						{
-							StructPathToIDTs.Empty(); // kill all cache
-							ReadConfigFile(ConfigFilepath);
-						}
-					}
-					else
-					{
-						ReadConfigFile(ConfigFilepath);
-					}
-				}
-			}
-			
-			return StructPathToIDTs.Find(ClassPath);
-		}
-
-		const TArray<TSharedPtr<FInstanceDataTransform>>* FindTransforms(const TNotNull<UStruct*> Struct)
-		{
-			if(Struct->GetOuter() != Struct->GetPackage())
-			{
-				// can't support transforms for nested structs like 
-				// anim blueprint's nested 'sparse class data':
-				return nullptr;
-			}
-			return FindTransforms(Struct->GetStructPathName());
-		}
-	private:
-
-		void ReadConfigFile(const FString& ConfigFilepath)
-		{
-			TSharedPtr<FJsonObject> ConfigJson;
-			FString JsonRaw;
-			if (FFileHelper::LoadFileToString(JsonRaw, *ConfigFilepath))
-			{
-				TSharedRef<TJsonReader<TCHAR>> JsonReader = TJsonReaderFactory<TCHAR>::Create(JsonRaw);
-				if (FJsonSerializer::Deserialize(JsonReader, ConfigJson))
-				{
-					FString PackagePathString;
-					if (!ConfigJson->TryGetStringField(TEXT("Package"), PackagePathString))
-					{
-						UE_LOG(LogInstanceDataTransforms, Warning, TEXT("'Package' feild was missing from '%s'"), *ConfigFilepath);
-						return;
-					}
-
-					FPackagePath PackagePath;
-					if (!FPackagePath::TryFromPackageName(PackagePathString, PackagePath))
-					{
-						UE_LOG(LogInstanceDataTransforms, Warning, TEXT("'ClassPackage' feild was malformed in '%s'"), *ConfigFilepath);
-						return;
-					}
-
-					const TArray<TSharedPtr<FJsonValue>>* ClassTransformsJson = nullptr;
-					if (!ConfigJson->TryGetArrayField(TEXT("ClassTransforms"), ClassTransformsJson) || !ClassTransformsJson)
-					{
-						UE_LOG(LogInstanceDataTransforms, Warning, TEXT("'ClassTransforms' array was missing or malformed in '%s'"), *ConfigFilepath);
-						return;
-					}
-
-					for (const TSharedPtr<FJsonValue>& TransformInfoJson : *ClassTransformsJson)
-					{
-						TSharedPtr<FJsonObject>* TransformInfo;
-						if (!TransformInfoJson->TryGetObject(TransformInfo) || !TransformInfo)
-						{
-							UE_LOG(LogInstanceDataTransforms, Warning, TEXT("Found a malfomed object in the 'ClassTransforms' array in '%s'"), *ConfigFilepath);
-							continue;
-						}
-
-						FTopLevelAssetPath StructPath;
-						if (FString ClassNameString; (*TransformInfo)->TryGetStringField(TEXT("Class"), ClassNameString))
-						{
-							FName ClassName = FName(ClassNameString);
-							StructPath = FTopLevelAssetPath(PackagePath.GetPackageFName(), ClassName);
-						}
-						else if (FString StructNameString; (*TransformInfo)->TryGetStringField(TEXT("Struct"), StructNameString))
-						{
-							FName StructName = FName(StructNameString);
-							StructPath = FTopLevelAssetPath(PackagePath.GetPackageFName(), StructName);
-						}
-						else
-						{
-							UE_LOG(LogInstanceDataTransforms, Warning, TEXT("an element of the 'ClassTransforms' array in '%s' was missing a 'Class' or 'Struct' field"), *ConfigFilepath);
-							continue;
-						}
-
-						FString OpCode;
-						if (!(*TransformInfo)->TryGetStringField(TEXT("OpCode"), OpCode))
-						{
-							UE_LOG(LogInstanceDataTransforms, Warning, TEXT("an element of the 'ClassTransforms' array in '%s' was missing a 'OpCode' field"), *ConfigFilepath);
-							continue;
-						}
-						const FInstanceDataTransformFunction* Func = RegisteredOperations.Find(OpCode);
-						if (!Func)
-						{
-							UE_LOG(LogInstanceDataTransforms, Warning, TEXT("The transform OpCode in '%s' was malformed. '%s' is not a valid operation."), *ConfigFilepath, *OpCode);
-							continue;
-						}
-						FInstanceDataTransformParams Params;
-						if (const FInstanceDataTransformParams* ParamsPtr; (*TransformInfo)->TryGetObjectField(TEXT("Params"), ParamsPtr))
-						{
-							Params = *ParamsPtr;
-						}
-						TArray<TSharedPtr<FInstanceDataTransform>>& ClassTransforms = StructPathToIDTs.FindOrAdd(StructPath);
-						ClassTransforms.Add(MakeShared<FInstanceDataTransform>(Params, *Func));
-					}
-				}
+				Directory.LeftInline(Found);
 			}
-			ConfigFileLastReadTime.Add(ConfigFilepath, FDateTime::Now());
-		}
+			Directory /= TEXT("Config");
 
-		friend struct FRegesterInstanceDataTransformHelper;
-		void RegisterOperation(FStringView Name, FInstanceDataTransformFunction Function)
-		{
-			RegisteredOperations.Add(FString(Name), Function);
+			return Directory / TEXT("InstanceDataTransforms.json");
 		}
-
-		TMap<FString, FDateTime> ConfigFileLastReadTime;
-		TMap<FTopLevelAssetPath, TArray<TSharedPtr<FInstanceDataTransform>>> StructPathToIDTs;
-		TMap<FString, FInstanceDataTransformFunction> RegisteredOperations;
-	};
+		return {};
+	}
 
 	struct FRegesterInstanceDataTransformHelper
 	{
-		FRegesterInstanceDataTransformHelper(FStringView Name, FInstanceDataTransformFunction Function)
+		FRegesterInstanceDataTransformHelper(FStringView Name, FInstanceDataTransforms::FTransformFunction Function)
 		{
-			FInstanceDataTransformRepository::Get().RegisterOperation(Name, Function);
+			FInstanceDataTransforms::Get().RegisterOperation(Name, Function);
 		}
 	};
 
@@ -375,67 +228,18 @@ namespace UE::Private
 		}
 	}
 
-	void PatchStructs(TObjectPtr<UObject>& IDO, TObjectPtr<UObject>& Owner)
+	bool SubPathIsLeaf(const TSharedPtr<FPropertyPathNameTree>& UnknownPropertyTree, FPropertyPathName PathName)
 	{
-		// search the unknown property tree for structs. If there are any, look up IDTs associated with that struct type and run them
-		TSharedPtr<FPropertyPathNameTree> UnknownPropertyTree = FUnknownPropertyTree(Owner).FindOrCreate();
-
-		// IDTs can modify the UnknownPropertyTree and iterating while modifying a tree is a no-no so pre-iterate and cache the paths that we want to visit ahead of time
-		TArray<FPropertyPathName> Paths;
-		FPropertyPathName BasePath;
-		FlattenPropertyTree(*UnknownPropertyTree, Paths, BasePath);
-
-		// search the unknown property tree for structs. If there are any, look up IDTs associated with that struct type and run them
-		for (int32 i = 0; i < Paths.Num(); ++i)
+		PathName.Pop();
+		while (PathName.GetSegmentCount() > 0)
 		{
-			FPropertyPathName& Path = Paths[i];
-			FPropertyPathNameSegment LeafSegment = Path.GetSegment(Path.GetSegmentCount() - 1);
-
-			auto PatchStruct = [&Path, &IDO, &Owner](const FPropertyTypeName& StructPath)
-				{
-					if (UStruct* Struct = UE::FindObjectByTypePath<UScriptStruct>(StructPath))
-					{
-
-						if (const TArray<TSharedPtr<Private::FInstanceDataTransform>>* IDTs = Private::FInstanceDataTransformRepository::Get().FindTransforms(Struct))
-						{
-							for (const TSharedPtr<Private::FInstanceDataTransform>& IDT : *IDTs)
-							{
-								Private::EInstanceDataTransformResult Result = IDT->Func(IDO, Owner, Path, IDT->Args);
-								if (Result == Private::EInstanceDataTransformResult::Break)
-								{
-									break;
-								}
-							}
-						}
-					}
-				};
-
-			if (LeafSegment.Type.GetName() == NAME_StructProperty)
-			{
-				PatchStruct(LeafSegment.Type.GetParameter(0));
-			}
-			else if (LeafSegment.Type.GetName() == NAME_ArrayProperty || LeafSegment.Type.GetName() == NAME_SetProperty)
-			{
-				FPropertyTypeName InnerType = LeafSegment.Type.GetParameter(0);
-				if (InnerType.GetName() == NAME_StructProperty)
-				{
-					PatchStruct(InnerType.GetParameter(0));
-				}
-			}
-			else if (LeafSegment.Type.GetName() == NAME_MapProperty)
+			if (UE::FPropertyPathNameTree::FNode Found = UnknownPropertyTree->Find(PathName))
 			{
-				FPropertyTypeName KeyType = LeafSegment.Type.GetParameter(0);
-				if (KeyType.GetName() == NAME_StructProperty)
-				{
-					PatchStruct(KeyType.GetParameter(0));
-				}
-				FPropertyTypeName ValType = LeafSegment.Type.GetParameter(1);
-				if (ValType.GetName() == NAME_StructProperty)
-				{
-					PatchStruct(ValType.GetParameter(0));
-				}
+				return !Found.GetSubTree();
 			}
+			PathName.Pop();
 		}
+		return false;
 	}
 }
 
@@ -449,7 +253,7 @@ namespace UE::Private::IDT
 	 * @param Params - contains a map with the following parameters
 	 *		- ["src"] the source FPropertyPathName as a json object. If this doesn't point to an unknown property, the operation will be skipped
 	 */
-	EInstanceDataTransformResult RemoveProperty(TObjectPtr<UObject>& IDO, TObjectPtr<UObject>& Owner, const FPropertyPathName& BasePath, const FInstanceDataTransformParams& Params)
+	ETransformResult RemoveProperty(TObjectPtr<UObject>& IDO, TObjectPtr<UObject>& Owner, const FPropertyPathName& BasePath, const FTransformParams& Params)
 	{
 		FUObjectSerializeContext* SerializeContext = FUObjectThreadContext::Get().GetSerializeContext();
 		TGuardValue<bool> ScopedImpersonateProperties(SerializeContext->bImpersonateProperties, true);
@@ -461,17 +265,34 @@ namespace UE::Private::IDT
 		if (!ensureAlwaysMsgf(Params && FPropertyPathName::FromJsonValue(Params->TryGetField(TEXT("src")), SrcPath),
 			TEXT("Failed to parse src parameter for a RemoveProperty IDT operation")))
 		{
-			return EInstanceDataTransformResult::Break;
+			return ETransformResult::Break;
 		}
 		InferMissingPathInfo(SrcPath, Owner->GetClass());
 
 		if (UnknownPropertyTree->Remove(SrcPath))
 		{
 			// create a new IDO with the modified property tree
+			TObjectPtr<UObject> OldIDO = IDO;
 			IDO = UE::FPropertyBagRepository::Get().ReinstanceWithCachedUnknownPropertyTree(Owner);
+#if WITH_EDITOR
+			FCoreUObjectDelegates::OnObjectsReinstanced.Broadcast({ {OldIDO, IDO} });
+#endif
+		}
+		else if (SubPathIsLeaf(UnknownPropertyTree, SrcPath))
+		{
+			// if a subpath was an overridden struct or container containing this property, we can't easily remove the property from the IDO
+			// class so instead we'll just remove the override on that value so it doesn't get saved anymore
+			FPropertyPathNameResolver ValueToClearOverrides(SrcPath, IDO);
+			if (ValueToClearOverrides)
+			{
+				TSharedPtr<FPropertyChangedChainEvent> OverrideChangeEvent;
+				ValueToClearOverrides.BuildChangeEvent(OverrideChangeEvent, EPropertyChangeType::ResetToDefault);
+				UE::SendClearOverriddenPropertyEvent(IDO, *OverrideChangeEvent, OverrideChangeEvent->PropertyChain);
+			}
 		}
 
-		return EInstanceDataTransformResult::Continue;
+
+		return ETransformResult::Continue;
 	};
 	REGISTER_IDT_OPERATION(RemoveProperty);
 
@@ -483,7 +304,7 @@ namespace UE::Private::IDT
 	 *		- ["src"] the source FPropertyPathName as a json object. If this doesn't point to an unknown property, the operation will be skipped
 	 *		- ["dst"] the destination FPropertyPathName as a json object. If this points to a non-existent property, the path will be added to the new IDO instance as an unknown property
 	 */
-	EInstanceDataTransformResult RedirectProperty(TObjectPtr<UObject>& IDO, TObjectPtr<UObject>& Owner, const FPropertyPathName& BasePath, const FInstanceDataTransformParams& Params)
+	ETransformResult RedirectProperty(TObjectPtr<UObject>& IDO, TObjectPtr<UObject>& Owner, const FPropertyPathName& BasePath, const FTransformParams& Params)
 	{
 		FUObjectSerializeContext* SerializeContext = FUObjectThreadContext::Get().GetSerializeContext();
 		TGuardValue<bool> ScopedImpersonateProperties(SerializeContext->bImpersonateProperties, true);
@@ -495,7 +316,7 @@ namespace UE::Private::IDT
 		if (!ensureAlwaysMsgf(Params && FPropertyPathName::FromJsonValue(Params->TryGetField(TEXT("src")), SrcPath),
 			TEXT("Failed to parse src parameter for a RedirectProperty IDT operation")))
 		{
-			return EInstanceDataTransformResult::Break;
+			return ETransformResult::Break;
 		}
 		InferMissingPathInfo(SrcPath, Owner->GetClass());
 
@@ -503,12 +324,11 @@ namespace UE::Private::IDT
 		if (!ensureAlwaysMsgf(Params && FPropertyPathName::FromJsonValue(Params->TryGetField(TEXT("dst")), DstPath),
 			TEXT("Failed to parse dst parameter for a RedirectProperty IDT operation")))
 		{
-			return EInstanceDataTransformResult::Break;
+			return ETransformResult::Break;
 		}
 		InferMissingPathInfo(DstPath, Owner->GetClass());
 
-
-		if (UnknownPropertyTree->Remove(SrcPath))
+		if (UnknownPropertyTree->Remove(SrcPath) || SubPathIsLeaf(UnknownPropertyTree, SrcPath))
 		{
 			FPropertyPathNameResolver DstInstanceValue(DstPath, Owner);
 			if (!DstInstanceValue)
@@ -522,6 +342,7 @@ namespace UE::Private::IDT
 			check(SrcValue); // it should be impossible for SrcPath to be in the old unknown property tree and not be in the old IDO
 
 			// create a new IDO with the modified property tree
+			TObjectPtr<UObject> OldIDO = IDO;
 			IDO = UE::FPropertyBagRepository::Get().ReinstanceWithCachedUnknownPropertyTree(Owner);
 
 			// find destination location in the new IDO
@@ -530,13 +351,23 @@ namespace UE::Private::IDT
 
 			while (SrcValue && DstIDOValue) // because of wildcards in the path, it's possible that there's several values to copy!
 			{
-				TSharedPtr<FPropertyChangedChainEvent> IdoChangeEvent;
-				DstIDOValue.BuildChangeEvent(IdoChangeEvent, EPropertyChangeType::ValueSet);
+				TSharedPtr<FPropertyChangedChainEvent> DstIdoChangeEvent;
+				DstIDOValue.BuildChangeEvent(DstIdoChangeEvent, EPropertyChangeType::ValueSet);
 
 				// copy property from SrcValue to DstIDOValue (note that because we're using change events, the value will be changed in the instance as well)
-				IDO->PreEditChange(IdoChangeEvent->PropertyChain);
+				IDO->PreEditChange(DstIdoChangeEvent->PropertyChain);
 				TryCopyOrConvert(SrcValue.Value, DstIDOValue.Value);
-				IDO->PostEditChangeChainProperty(*IdoChangeEvent);
+				IDO->PostEditChangeChainProperty(*DstIdoChangeEvent);
+
+				// if a subpath was an overridden struct or container containing this property, then we can't easily remove the property from the IDO
+				// class so instead we'll just remove the override on that value so it doesn't get saved anymore
+				FPropertyPathNameResolver ValueToClearOverrides(SrcPath, IDO);
+				if (ValueToClearOverrides)
+				{
+					TSharedPtr<FPropertyChangedChainEvent> OverrideChangeEvent;
+					ValueToClearOverrides.BuildChangeEvent(OverrideChangeEvent, EPropertyChangeType::ResetToDefault);
+					UE::SendClearOverriddenPropertyEvent(IDO, *OverrideChangeEvent, OverrideChangeEvent->PropertyChain);
+				}
 
 				/*
 					because containers may contain more than one instance of a struct being redirected, we redirect them all by using wildcards in the path. 
@@ -550,9 +381,13 @@ namespace UE::Private::IDT
 			// this should never fail unless a user has used wildcards in the path they manually wrote into the IDT
 			ensureMsgf(!SrcValue && !DstIDOValue, TEXT("An IDT ran that operates on every element of a container but the source and destination had a " 
 				"different number of elements! This may have caused corruption in your data. Save at your own risk."));
+
+#if WITH_EDITOR
+			FCoreUObjectDelegates::OnObjectsReinstanced.Broadcast({ {OldIDO, IDO} });
+#endif
 		}
 
-		return EInstanceDataTransformResult::Continue;
+		return ETransformResult::Continue;
 	};
 	REGISTER_IDT_OPERATION(RedirectProperty);
 }
@@ -560,31 +395,300 @@ namespace UE::Private::IDT
 
 namespace UE
 {
-	bool PatchInstanceDataObject(TObjectPtr<UObject>& IDO, TObjectPtr<UObject>& Owner)
+	FInstanceDataTransforms& FInstanceDataTransforms::Get()
+	{
+		static FInstanceDataTransforms Instance;
+		return Instance;
+	}
+
+	void FInstanceDataTransforms::WriteInstanceDataTransform(UStruct& ClassOrStruct, const FString& OpCode, const FTransformParams& Params)
+	{
+		FTopLevelAssetPath AssetPath = ClassOrStruct.GetStructPathName();
+		FString ConfigFilepath = Private::InstanceDataTransformDirectory(AssetPath);
+		if (ConfigFilepath.IsEmpty())
+		{
+			UE_LOG(LogInstanceDataTransforms, Error, TEXT("Couldn't find a valid IDT config path for '%s'. Perhaps it's not mounted?"), *AssetPath.ToString());
+			return;
+		}
+
+		FString JsonRaw;
+		TSharedPtr<FJsonObject> ConfigJson;
+		if (FFileHelper::LoadFileToString(JsonRaw, *ConfigFilepath))
+		{
+			TSharedRef<TJsonReader<TCHAR>> JsonReader = TJsonReaderFactory<TCHAR>::Create(JsonRaw);
+			if (!FJsonSerializer::Deserialize(JsonReader, ConfigJson))
+			{
+				UE_LOG(LogInstanceDataTransforms, Error, TEXT("'%s' is malformed and could not deserialize."), *ConfigFilepath);
+				return;
+			}
+		}
+		else
+		{
+			ConfigJson = MakeShared<FJsonObject>();
+			ConfigJson->SetStringField(TEXT("Package"), AssetPath.GetPackageName().ToString());
+			ConfigJson->SetArrayField(TEXT("Transforms"), {});
+		}
+
+		const TArray<TSharedPtr<FJsonValue>>* ExistingTransformsPtr = nullptr;
+		if (!ConfigJson->TryGetArrayField(TEXT("Transforms"), ExistingTransformsPtr) || !ExistingTransformsPtr)
+		{
+			UE_LOG(LogInstanceDataTransforms, Error, TEXT("'Transforms' array was missing or malformed in '%s'"), *ConfigFilepath);
+			return;
+		}
+
+		// create the new transform operation
+		TArray<TSharedPtr<FJsonValue>> TransformsJson = *ExistingTransformsPtr;
+		TSharedPtr<FJsonObject> NewTransform = MakeShared<FJsonObject>();
+		if (ClassOrStruct.IsA<UClass>())
+		{
+			NewTransform->SetStringField(TEXT("Class"), AssetPath.GetAssetName().ToString());
+		}
+		else
+		{
+			NewTransform->SetStringField(TEXT("Struct"), AssetPath.GetAssetName().ToString());
+		}
+		NewTransform->SetStringField(TEXT("OpCode"), OpCode);
+		if (Params)
+		{
+			NewTransform->SetObjectField(TEXT("Params"), Params);
+		}
+
+		// append the new operation to the Transforms list
+		TransformsJson.Add(StaticCastSharedRef<FJsonValue>(MakeShared<FJsonValueObject>(NewTransform)));
+		ConfigJson->SetArrayField(TEXT("Transforms"), TransformsJson);
+
+		// update Json file on disk
+		JsonRaw.Empty();
+		TSharedRef<TJsonWriter<TCHAR>> JsonWriter = TJsonWriterFactory<TCHAR>::Create(&JsonRaw);
+		if (!FJsonSerializer::Serialize(ConfigJson.ToSharedRef(), JsonWriter))
+		{
+			UE_LOG(LogInstanceDataTransforms, Error, TEXT("'%s' failed to Serialize json for."), *ConfigFilepath);
+			return;
+		}
+
+		if (!FFileHelper::SaveStringToFile(JsonRaw, *ConfigFilepath))
+		{
+			UE_LOG(LogInstanceDataTransforms, Error, TEXT("'%s' failed to save."), *ConfigFilepath);
+			return;
+		}
+	}
+
+	TObjectPtr<UObject> FInstanceDataTransforms::PatchInstanceDataObject(TNotNull<UObject*> IDO, TNotNull<UObject*> Owner)
+	{
+		TObjectPtr<UObject> ResultIdo = IDO;
+		TObjectPtr<UObject> ResultOwner = Owner;
+		PatchClass(ResultIdo, ResultOwner);
+		PatchStructs(ResultIdo, ResultOwner);
+
+		return ResultIdo;
+	}
+
+	const TArray<TSharedPtr<FInstanceDataTransforms::FTransform>>* FInstanceDataTransforms::FindTransformsInternal(const FTopLevelAssetPath& ClassPath)
+	{
+		// the transform config for a class is found in <RootPackageDirectory>/Config/InstanceDataTransforms.json
+		// TODO: I need to some research on wheather any work needs to be done in regards to making sure this file is treated properly by source control and the content worker.
+		FString ConfigFilepath = Private::InstanceDataTransformDirectory(ClassPath);
+		if (IFileManager::Get().FileExists(*ConfigFilepath))
+		{
+			FDateTime FileModificationTime = IFileManager::Get().GetTimeStamp(*ConfigFilepath);
+
+			if (FDateTime* LastReadTime = ConfigFileLastReadTime.Find(ConfigFilepath))
+			{
+				// check if the config file has been updated since the last time we read it
+				if (*LastReadTime < FileModificationTime)
+				{
+					StructPathToIDTs.Empty(); // kill all cache
+					ReadConfigFile(ConfigFilepath);
+				}
+			}
+			else
+			{
+				ReadConfigFile(ConfigFilepath);
+			}
+		}
+
+		return StructPathToIDTs.Find(ClassPath);
+	}
+
+	const TArray<TSharedPtr<FInstanceDataTransforms::FTransform>>* FInstanceDataTransforms::FindTransformsInternal(const TNotNull<UStruct*> Struct)
+	{
+		if(Struct->GetOuter() != Struct->GetPackage())
+		{
+			// can't support transforms for nested structs like 
+			// anim blueprint's nested 'sparse class data':
+			return nullptr;
+		}
+		return FindTransformsInternal(Struct->GetStructPathName());
+	}
+
+	void FInstanceDataTransforms::ReadConfigFile(const FString& ConfigFilepath)
+	{
+		TSharedPtr<FJsonObject> ConfigJson;
+		FString JsonRaw;
+		if (FFileHelper::LoadFileToString(JsonRaw, *ConfigFilepath))
+		{
+			TSharedRef<TJsonReader<TCHAR>> JsonReader = TJsonReaderFactory<TCHAR>::Create(JsonRaw);
+			if (FJsonSerializer::Deserialize(JsonReader, ConfigJson))
+			{
+				FString PackagePathString;
+				if (!ConfigJson->TryGetStringField(TEXT("Package"), PackagePathString))
+				{
+					UE_LOG(LogInstanceDataTransforms, Warning, TEXT("'Package' feild was missing from '%s'"), *ConfigFilepath);
+					return;
+				}
+
+				FPackagePath PackagePath;
+				if (!FPackagePath::TryFromPackageName(PackagePathString, PackagePath))
+				{
+					UE_LOG(LogInstanceDataTransforms, Warning, TEXT("'ClassPackage' feild was malformed in '%s'"), *ConfigFilepath);
+					return;
+				}
+
+				const TArray<TSharedPtr<FJsonValue>>* TransformsJson = nullptr;
+				if (!ConfigJson->TryGetArrayField(TEXT("Transforms"), TransformsJson) || !TransformsJson)
+				{
+					UE_LOG(LogInstanceDataTransforms, Warning, TEXT("'Transforms' array was missing or malformed in '%s'"), *ConfigFilepath);
+					return;
+				}
+
+				for (const TSharedPtr<FJsonValue>& TransformInfoJson : *TransformsJson)
+				{
+					TSharedPtr<FJsonObject>* TransformInfo;
+					if (!TransformInfoJson->TryGetObject(TransformInfo) || !TransformInfo)
+					{
+						UE_LOG(LogInstanceDataTransforms, Warning, TEXT("Found a malfomed object in the 'Transforms' array in '%s'"), *ConfigFilepath);
+						continue;
+					}
+
+					FTopLevelAssetPath StructPath;
+					if (FString ClassNameString; (*TransformInfo)->TryGetStringField(TEXT("Class"), ClassNameString))
+					{
+						FName ClassName = FName(ClassNameString);
+						StructPath = FTopLevelAssetPath(PackagePath.GetPackageFName(), ClassName);
+					}
+					else if (FString StructNameString; (*TransformInfo)->TryGetStringField(TEXT("Struct"), StructNameString))
+					{
+						FName StructName = FName(StructNameString);
+						StructPath = FTopLevelAssetPath(PackagePath.GetPackageFName(), StructName);
+					}
+					else
+					{
+						UE_LOG(LogInstanceDataTransforms, Warning, TEXT("an element of the 'Transforms' array in '%s' was missing a 'Class' or 'Struct' field"), *ConfigFilepath);
+						continue;
+					}
+
+					FString OpCode;
+					if (!(*TransformInfo)->TryGetStringField(TEXT("OpCode"), OpCode))
+					{
+						UE_LOG(LogInstanceDataTransforms, Warning, TEXT("an element of the 'Transforms' array in '%s' was missing a 'OpCode' field"), *ConfigFilepath);
+						continue;
+					}
+					const FTransformFunction* Func = RegisteredOperations.Find(OpCode);
+					if (!Func)
+					{
+						UE_LOG(LogInstanceDataTransforms, Warning, TEXT("The transform OpCode in '%s' was malformed. '%s' is not a valid operation."), *ConfigFilepath, *OpCode);
+						continue;
+					}
+					FTransformParams Params;
+					if (const FTransformParams* ParamsPtr; (*TransformInfo)->TryGetObjectField(TEXT("Params"), ParamsPtr))
+					{
+						Params = *ParamsPtr;
+					}
+					TArray<TSharedPtr<FTransform>>& Transforms = StructPathToIDTs.FindOrAdd(StructPath);
+					Transforms.Add(MakeShared<FTransform>(Params, *Func));
+				}
+			}
+		}
+		ConfigFileLastReadTime.Add(ConfigFilepath, FDateTime::Now());
+	}
+
+	void FInstanceDataTransforms::RegisterOperation(FStringView Name, FTransformFunction Function)
 	{
-		check(IDO);
-		check(Owner);
+		RegisteredOperations.Add(FString(Name), Function);
+	}
 
+	void FInstanceDataTransforms::PatchClass(TObjectPtr<UObject>& IDO, TObjectPtr<UObject>& Owner)
+	{
 		// find and run class IDTs
-		if (const TArray<TSharedPtr<Private::FInstanceDataTransform>>* IDTs = Private::FInstanceDataTransformRepository::Get().FindTransforms(Owner->GetClass()))
+		if (const TArray<TSharedPtr<FTransform>>* IDTs = FindTransformsInternal(Owner->GetClass()))
 		{
-			for (const TSharedPtr<Private::FInstanceDataTransform>& IDT : *IDTs)
+			for (const TSharedPtr<FTransform>& IDT : *IDTs)
 			{
 				FPropertyPathName RootPath; // class IDTs operate at the root of the class
-				Private::EInstanceDataTransformResult Result = IDT->Func(IDO, Owner, RootPath, IDT->Args);
-				if (Result == Private::EInstanceDataTransformResult::Break)
+				ETransformResult Result = IDT->Func(IDO, Owner, RootPath, IDT->Args);
+				if (Result == ETransformResult::Break)
 				{
 					break;
 				}
 			}
 		}
+	}
 
-		// find and run struct IDTs
-		Private::PatchStructs(IDO, Owner);
+	void FInstanceDataTransforms::PatchStructs(TObjectPtr<UObject>& IDO, TObjectPtr<UObject>& Owner)
+	{
+
+		// search the unknown property tree for structs. If there are any, look up IDTs associated with that struct type and run them
+		TSharedPtr<FPropertyPathNameTree> UnknownPropertyTree = FUnknownPropertyTree(Owner).FindOrCreate();
+
+		// IDTs can modify the UnknownPropertyTree and iterating while modifying a tree is a no-no so pre-iterate and cache the paths that we want to visit ahead of time
+		TArray<FPropertyPathName> Paths;
+		FPropertyPathName BasePath;
+		Private::FlattenPropertyTree(*UnknownPropertyTree, Paths, BasePath);
+
+		// search the unknown property tree for structs. If there are any, look up IDTs associated with that struct type and run them
+		for (int32 i = 0; i < Paths.Num(); ++i)
+		{
+			FPropertyPathName& Path = Paths[i];
+			FPropertyPathNameSegment LeafSegment = Path.GetSegment(Path.GetSegmentCount() - 1);
+
+			auto PatchStruct = [&Path, &IDO, &Owner](const FPropertyTypeName& StructPath)
+				{
+					if (UStruct* Struct = UE::FindObjectByTypePath<UScriptStruct>(StructPath))
+					{
+
+						if (const TArray<TSharedPtr<FTransform>>* IDTs = FInstanceDataTransforms::Get().FindTransformsInternal(Struct))
+						{
+							for (const TSharedPtr<FTransform>& IDT : *IDTs)
+							{
+								ETransformResult Result = IDT->Func(IDO, Owner, Path, IDT->Args);
+								if (Result == ETransformResult::Break)
+								{
+									break;
+								}
+							}
+						}
+					}
+				};
 
-		return true;
+			if (LeafSegment.Type.GetName() == NAME_StructProperty)
+			{
+				PatchStruct(LeafSegment.Type.GetParameter(0));
+			}
+			else if (LeafSegment.Type.GetName() == NAME_ArrayProperty || LeafSegment.Type.GetName() == NAME_SetProperty)
+			{
+				FPropertyTypeName InnerType = LeafSegment.Type.GetParameter(0);
+				if (InnerType.GetName() == NAME_StructProperty)
+				{
+					PatchStruct(InnerType.GetParameter(0));
+				}
+			}
+			else if (LeafSegment.Type.GetName() == NAME_MapProperty)
+			{
+				FPropertyTypeName KeyType = LeafSegment.Type.GetParameter(0);
+				if (KeyType.GetName() == NAME_StructProperty)
+				{
+					PatchStruct(KeyType.GetParameter(0));
+				}
+				FPropertyTypeName ValType = LeafSegment.Type.GetParameter(1);
+				if (ValType.GetName() == NAME_StructProperty)
+				{
+					PatchStruct(ValType.GetParameter(0));
+				}
+			}
+		}
 	}
 
+
 }
 
 #endif // !UE_BUILD_SHIPPING_WITH_EDITOR && WITH_EDITORONLY_DATA 
\ No newline at end of file
diff --git a/Engine/Source/Runtime/CoreUObject/Private/UObject/PropertyBagRepository.cpp b/Engine/Source/Runtime/CoreUObject/Private/UObject/PropertyBagRepository.cpp
--- a/Engine/Source/Runtime/CoreUObject/Private/UObject/PropertyBagRepository.cpp
+++ b/Engine/Source/Runtime/CoreUObject/Private/UObject/PropertyBagRepository.cpp
@@ -23,6 +23,7 @@
 #include "UObject/UObjectHash.h"
 #include "UObject/UObjectSerializeContext.h"
 #include "UObject/UObjectThreadContext.h"
+#include "UObject/PropertyPathNameTree.h"
 #include "Templates/UnrealTemplate.h"
 
 #if WITH_EDITOR
@@ -211,6 +212,12 @@ void FPropertyBagRepository::ReassociateObjects(const TMap<UObject*, UObject*>&
 				// inferred some info. Like for example if an array type mismatch occured and the new array was auto-grown in the owner to match the
 				// old one. This array growth needs to be copied over to the IDO to maintain it
 				CopyTaggedProperties(Pair.Value, NewBagData->InstanceDataObject);
+
+#if !UE_BUILD_SHIPPING_WITH_EDITOR
+				// Run IDT operations on this instance to recover data
+				TObjectPtr<UObject> InOutInstance = const_cast<UObject*>(Pair.Value);
+				UE::FInstanceDataTransforms::Get().PatchInstanceDataObject(NewBagData->InstanceDataObject, InOutInstance);
+#endif
 			}
 
 			OldBagData.Destroy();
@@ -685,7 +692,7 @@ void FPropertyBagRepository::PostLoadInstanceDataObject(const UObject* Owner)
 #if !UE_BUILD_SHIPPING_WITH_EDITOR
 			// Run IDT operations on this instance to recover data
 			TObjectPtr<UObject> InOutInstance = const_cast<UObject*>(Owner);
-			UE::PatchInstanceDataObject(BagData->InstanceDataObject, InOutInstance);
+			UE::FInstanceDataTransforms::Get().PatchInstanceDataObject(BagData->InstanceDataObject, InOutInstance);
 #endif
 		}
 	}
@@ -700,17 +707,15 @@ void FPropertyBagRepository::DestroyOuterBag(const UObject* Owner)
 
 bool FPropertyBagRepository::RequiresFixup(const UObject* Object, bool bIncludeOuter) const
 {
-	FPropertyBagRepositoryLock LockRepo(this);
-
-	const FPropertyBagAssociationData* BagData = AssociatedData.Find(Object);
-	bool bResult = BagData ? BagData->bNeedsFixup : false;
+	TSharedPtr<FPropertyPathNameTree> UnknownPropertyTree = FUnknownPropertyTree(Object).Find();
+	bool bResult = UnknownPropertyTree ? !UnknownPropertyTree->IsEmpty() : false;
 	if (!bResult && bIncludeOuter)
 	{
 		ForEachObjectWithOuterBreakable(Object,
 			[&bResult, this](UObject* Object) 
 			{
-				if (const FPropertyBagAssociationData* BagData = AssociatedData.Find(Object); 
-					BagData && BagData->bNeedsFixup)
+				TSharedPtr<FPropertyPathNameTree> UnknownPropertyTree = FUnknownPropertyTree(Object).Find();
+				if (UnknownPropertyTree && !UnknownPropertyTree->IsEmpty())
 				{
 					bResult = true;
 					return false;
@@ -721,14 +726,30 @@ bool FPropertyBagRepository::RequiresFixup(const UObject* Object, bool bIncludeO
 	return bResult;
 }
 
-
-void FPropertyBagRepository::MarkAsFixedUp(const UObject* Object)
+bool FPropertyBagRepository::RequiresFixup(const UObject* Object, const FPropertyPathName& Path) const
 {
-	FPropertyBagRepositoryLock LockRepo(this);
-	if (FPropertyBagAssociationData* BagData = AssociatedData.Find(Object))
+	TSharedPtr<FPropertyPathNameTree> UnknownPropertyTree = FUnknownPropertyTree(Object).Find();
+	if (!UnknownPropertyTree || UnknownPropertyTree->IsEmpty())
 	{
-		BagData->bNeedsFixup = false;
+		return false;
 	}
+
+	FUObjectSerializeContext* SerializeContext = FUObjectThreadContext::Get().GetSerializeContext();
+	TGuardValue<bool> ScopedImpersonateProperties(SerializeContext->bImpersonateProperties, true);
+
+	// look for either the full path in the tree, or a leaf node in the tree that matches the beginning of the path.
+	// the latter scenario can occur because if a struct is a leaf, all it's child nodes are implicitly also unknown properties so they're not included in the tree
+	UE::FPropertyPathName SubPath = Path;
+	while (SubPath.GetSegmentCount() > 0)
+	{
+		if (UE::FPropertyPathNameTree::FNode Found = UnknownPropertyTree->Find(SubPath))
+		{
+			return (SubPath.GetSegmentCount() == Path.GetSegmentCount()) || (!Found.GetSubTree());
+		}
+		SubPath.Pop();
+	}
+
+	return false;
 }
 
 bool FPropertyBagRepository::RemoveAssociationUnsafe(const UObject* Owner)
@@ -774,7 +795,7 @@ void FPropertyBagRepository::FindNestedInstanceDataObject(const UObject* Owner,
 	FPropertyBagRepositoryLock LockRepo(this);
 
 	if (const FPropertyBagAssociationData* BagData = AssociatedData.Find(Owner); 
-		BagData && BagData->InstanceDataObject && (!bRequiresFixupOnly || BagData->bNeedsFixup))
+		BagData && BagData->InstanceDataObject && (!bRequiresFixupOnly || RequiresFixup(Owner, false)))
 	{
 		Callback(BagData->InstanceDataObject);
 	}
@@ -783,7 +804,7 @@ void FPropertyBagRepository::FindNestedInstanceDataObject(const UObject* Owner,
 		[this, bRequiresFixupOnly, Callback](UObject* Object)
 		{
 			if (const FPropertyBagAssociationData* BagData = AssociatedData.Find(Object);
-				BagData && BagData->InstanceDataObject && (!bRequiresFixupOnly || BagData->bNeedsFixup))
+				BagData && BagData->InstanceDataObject && (!bRequiresFixupOnly || RequiresFixup(Object, false)))
 			{
 				Callback(BagData->InstanceDataObject);
 			}
@@ -853,11 +874,7 @@ UObject* FPropertyBagRepository::ReinstanceWithCachedUnknownPropertyTree(UObject
 		// replace the old IDO with the new one in the repository
 		InstanceDataObjectToOwner.Remove(BagData->InstanceDataObject);
 		BagData->InstanceDataObject = NewIDO;
-		BagData->bNeedsFixup = UE::StructContainsLooseProperties(NewIdoClass);
 		InstanceDataObjectToOwner.Add(NewIDO, Owner);
-#if WITH_EDITOR
-		FCoreUObjectDelegates::OnObjectsReinstanced.Broadcast({{IDO, NewIDO}});
-#endif
 
 		return NewIDO;
 	}
@@ -924,7 +941,6 @@ UObject* FPropertyBagRepository::CreateInstanceDataObject(UObject* Owner, FArchi
 
 	FPropertyBagAssociationData& BagData = AssociatedData.Add(Owner);
 	BagData.InstanceDataObject = InstanceDataObject;
-	BagData.bNeedsFixup = StructContainsLooseProperties(InstanceDataObjectClass);
 
 	// Load the IDO.
 	if (StartOffset != EndOffset)
@@ -1087,7 +1103,7 @@ void FPropertyBagRepository::GatherStats(FPropertyBagRepositoryStats& Stats)
 			FArchiveCountMem MemoryCount(BagData.InstanceDataObject);
 			Stats.IDOMemoryBytes += MemoryCount.GetMax();
 
-			if (BagData.bNeedsFixup)
+			if (RequiresFixup(Pair.Key))
 			{
 				++Stats.NumIDOsWithLooseProperties;
 			}
