DRT integration into Ido Fixup Tool

diff --git a/Engine/Plugins/Experimental/InstanceDataObjectFixupTool/Source/InstanceDataObjectFixupTool.Build.cs b/Engine/Plugins/Experimental/InstanceDataObjectFixupTool/Source/InstanceDataObjectFixupTool.Build.cs
--- a/Engine/Plugins/Experimental/InstanceDataObjectFixupTool/Source/InstanceDataObjectFixupTool.Build.cs
+++ b/Engine/Plugins/Experimental/InstanceDataObjectFixupTool/Source/InstanceDataObjectFixupTool.Build.cs
@@ -23,7 +23,15 @@ namespace UnrealBuildTool.Rules
 					"UnrealEd",
 					"TedsAlerts",
 					"TypedElementFramework",
-					"Json"
+					"Json",
+					"TedsTypeInfo",
+					"TedsQueryStack",
+					"TedsTableViewer",
+					"InputCore",
+					"ToolWidgets",
+					"TedsOutliner",
+					"SceneOutliner",
+					"EditorWidgets",
 				}
 			);
 		}
diff --git a/Engine/Plugins/Experimental/InstanceDataObjectFixupTool/Source/Private/DataRecoveryTool.cpp b/Engine/Plugins/Experimental/InstanceDataObjectFixupTool/Source/Private/DataRecoveryTool.cpp
new file mode 100644
--- /dev/null
+++ b/Engine/Plugins/Experimental/InstanceDataObjectFixupTool/Source/Private/DataRecoveryTool.cpp
@@ -0,0 +1,409 @@
+// Copyright Epic Games, Inc. All Rights Reserved.
+
+#include "DataRecoveryTool.h"
+
+#include "DataRecoveryToolStyleSheet.h"
+#include "DataStorage/Features.h"
+#include "DetailTreeNode.h"
+#include "InstanceDataObjectFixupTool.h"
+#include "Elements/Columns/TedsTypeInfoColumns.h"
+#include "Elements/Columns/TypedElementCompatibilityColumns.h"
+#include "Elements/Columns/TypedElementHiearchyColumns.h"
+#include "Elements/Columns/TypedElementTypeInfoColumns.h"
+#include "Elements/Framework/TypedElementQueryBuilder.h"
+#include "Elements/Interfaces/TypedElementDataStorageInterface.h"
+#include "Elements/Interfaces/TypedElementDataStorageUiInterface.h"
+#include "Framework/MultiBox/MultiBoxBuilder.h"
+#include "SDetailsSplitter.h"
+#include "SlateOptMacros.h"
+#include "SPositiveActionButton.h"
+#include "TedsQueryNode.h"
+#include "TedsQueryStack/Public/TedsRowQueryResultsNode.h"
+#include "TedsTableViewer/Public/Widgets/STedsSearchBox.h"
+#include "TedsTableViewer/Public/Widgets/STedsTableViewer.h"
+#include "Widgets/Input/SButton.h"
+#include "Widgets/Layout/SSeparator.h"
+#include "Widgets/STedsHierarchyViewer.h"
+#include "Serialization/ObjectWriter.h"
+
+
+#define LOCTEXT_NAMESPACE "DataRecoveryTool"
+
+BEGIN_SLATE_FUNCTION_BUILD_OPTIMIZATION
+
+using namespace UE::DataRecoveryTool::StyleSheet;
+
+void SDataRecoveryTool::Construct(const FArguments& InArgs)
+{
+	SelectedClassPath =  MakeShared<FTopLevelAssetPath>(InArgs._SelectedClassPath.IsSet() ? *InArgs._SelectedClassPath : nullptr);
+	ChildSlot
+		[
+			SNew(SVerticalBox)
+			+SVerticalBox::Slot()
+			.AutoHeight()
+			[
+				MakeTopBar()
+			]
+			+SVerticalBox::Slot()
+			.FillHeight(1.f)
+			.Padding(Padding::None, Padding::Small, Padding::None, Padding::None)
+			[
+				SNew(SSplitter)
+				+SSplitter::Slot()
+				.MinSize(290.f)
+				.Value(.3f)
+				[
+					MakeClassesPanel()
+				]
+				+ SSplitter::Slot()
+				.MinSize(290.f)
+				[
+					MakeRightPanel()
+				]
+			]
+		];
+
+}
+
+
+
+TSharedRef<SWidget> SDataRecoveryTool::MakeTopBar()
+{
+	return
+		SNew(SBorder)
+		.BorderImage(FAppStyle::Get().GetBrush("Brushes.Panel"))
+		.Padding(Padding::None)
+		[
+			SNew(SHorizontalBox)
+			+SHorizontalBox::Slot()
+			.AutoWidth()
+			.Padding(0)
+			[
+				SNew(SButton)
+				.HAlign(HAlign_Center)
+				.VAlign(VAlign_Center)
+				.ContentPadding(Padding::Normal)
+				.ButtonStyle(FAppStyle::Get(), "SimpleButton")
+				.OnClicked(this, &SDataRecoveryTool::OnSaveClicked)
+				[
+					SNew(SImage)
+					.Image(FAppStyle::GetBrush("Icons.Save"))
+					.ToolTipText(INVTEXT("Save"))
+				]
+			]
+			+ SHorizontalBox::Slot()
+				.AutoWidth()
+				.VAlign(VAlign_Fill)
+				[
+					SNew(SSeparator)
+					.Thickness(Padding::Small * .5f)
+					.Orientation(Orient_Vertical)
+					.SeparatorImage(FAppStyle::Get().GetBrush("Brushes.Background"))
+				]
+				+ SHorizontalBox::Slot()
+				.FillWidth(1.f)
+				[
+					SNew(SHorizontalBox)
+					+ SHorizontalBox::Slot()
+					.AutoWidth()
+					.HAlign(HAlign_Left)
+					[
+						SNew(SButton)
+						.HAlign(HAlign_Center)
+						.VAlign(VAlign_Center)
+						.ContentPadding(Padding::Normal)
+						.ButtonStyle(FAppStyle::Get(), "SimpleButton")
+						[
+							SNew(SHorizontalBox)
+							+ SHorizontalBox::Slot()
+							.AutoWidth()
+							[
+								SNew(SImage)
+								.Image(FAppStyle::GetBrush("LevelEditor.Recompile.Small"))
+							]
+							+ SHorizontalBox::Slot()
+							.Padding(Padding::Small, Padding::None, Padding::None, Padding::None)
+							.AutoWidth()
+							[
+								SNew(STextBlock)
+								.Text(LOCTEXT("ApplyButtonText", "Apply"))
+							]
+						]
+					]
+					+ SHorizontalBox::Slot()
+					.FillWidth(1.f)
+					[
+						SNullWidget::NullWidget
+					]
+
+					+ SHorizontalBox::Slot()
+					.HAlign(HAlign_Right)
+					[
+						SNew(SButton)
+						.HAlign(HAlign_Center)
+						.VAlign(VAlign_Center)
+						.ContentPadding(Padding::Normal)
+						.ButtonStyle(FAppStyle::Get(), "SimpleButton")
+					[
+						SNew(SHorizontalBox)
+						+ SHorizontalBox::Slot()
+						.AutoWidth()
+						[
+							SNew(SImage)
+							.Image(GetOpenExternalImage())
+						]
+						+ SHorizontalBox::Slot()
+						.Padding(Padding::Small, Padding::None, Padding::None, Padding::None)
+						.AutoWidth()
+						[
+							SNew(STextBlock)
+							.Text(LOCTEXT("EditInVSCodeButton", "Edit in VS Code"))
+						]
+					]
+				]
+			]
+		];
+}
+
+
+
+TSharedRef<SWidget> SDataRecoveryTool::MakeClassesPanel()
+{
+	return MakePanel(LOCTEXT("ClassesTitle", "CLASSES"), MakeTEDSTableViewer());
+}
+
+TSharedRef<SWidget> SDataRecoveryTool::MakeRightPanel()
+{
+	return MakePanel(LOCTEXT("MappingsTitle", "MAPPINGS"), MakeDiffSection());
+}
+
+FReply SDataRecoveryTool::OnSaveClicked()
+{
+	//TODO: Implement this
+	return FReply::Handled();
+}
+
+
+TSharedRef<SWidget> SDataRecoveryTool::MakeDiffSection()
+{
+	return
+		SNew(SBorder)
+		.BorderImage(FAppStyle::Get().GetBrush("Brushes.Recessed"))
+		.Padding(Padding::None)
+		[
+			SAssignNew(DiffView, SInstanceDataObjectFixupTool)
+				.SelectedClassPath(SelectedClassPath)
+		];
+}
+
+
+TSharedRef<SWidget> SDataRecoveryTool::MakePanel(const FText& InTitle, const TSharedRef<SWidget>& InContent, const TOptional<FOnGetContent>& InMenuContent)
+{
+	TSharedRef<SHorizontalBox> Header =
+		SNew(SHorizontalBox)
+		+ SHorizontalBox::Slot()
+		.AutoWidth()
+		.VAlign(VAlign_Center)
+		[
+			SNew(STextBlock)
+			.Text(InTitle)
+			.Font(FAppStyle::Get().GetFontStyle("NormalFontBold"))
+		];
+
+	if (InMenuContent.IsSet())
+	{
+		Header
+		->AddSlot()
+		.FillWidth(1.f)
+		[
+			SNullWidget::NullWidget
+		];
+
+		Header
+		->AddSlot()
+		.AutoWidth()
+		.VAlign(VAlign_Center)
+		[
+			SNew(SPositiveActionButton)
+			.Icon(FAppStyle::Get().GetBrush("Icons.Plus"))
+			.Text(LOCTEXT("AddMappingButton", "Add"))
+			.OnGetMenuContent(InMenuContent.GetValue())
+		];
+	}
+
+	TSharedRef<SWidget> HeaderTitle =
+		SNew(SBorder)
+		.BorderImage(FAppStyle::Get().GetBrush("Brushes.Panel"))
+		.Padding(Padding::Big)
+		[
+			Header
+		];
+
+	return
+		SNew(SBorder)
+		.Padding(Padding::None)
+		.BorderImage(FAppStyle::Get().GetBrush("Brushes.Panel"))
+		[
+			SNew(SVerticalBox)
+			+ SVerticalBox::Slot()
+			.AutoHeight()
+			[
+				HeaderTitle
+			]
+			+ SVerticalBox::Slot()
+			.FillHeight(1.f)
+			.FillContentHeight(1.f)
+			[
+				InContent
+			]
+		];
+}
+
+void SDataRecoveryTool::OnListSelectionChanged(UE::Editor::DataStorage::RowHandle RowHandle)
+{
+	using namespace UE::Editor::DataStorage;
+
+	if (RowHandle == 0)
+	{
+		*SelectedClassPath = nullptr;
+
+		DiffView->GenerateDetailsViews();
+
+		return;
+	}
+
+	ICoreProvider* Storage = GetMutableDataStorageFeature<ICoreProvider>(StorageFeatureName);
+
+	FTypedElementClassTypeInfoColumn* TypeColumn = Storage->GetColumn<FTypedElementClassTypeInfoColumn>(RowHandle);
+
+	if (!ensureMsgf(TypeColumn, TEXT("Should have been able to find a row")))
+	{
+		return;
+	}
+
+	const UClass* Type = TypeColumn->TypeInfo.Get();
+	if (!ensureMsgf(Type, TEXT("Should have been able to find the Type UClass")))
+	{
+		return;
+	}
+
+	*SelectedClassPath = Type->GetClassPathName();
+
+	if (DiffView == nullptr)
+	{
+		return;
+	}
+
+	DiffView->GenerateDetailsViews();
+}
+
+void SDataRecoveryTool::SetSelectedClass()
+{
+	FTSTicker::GetCoreTicker().AddTicker(
+		FTickerDelegate::CreateLambda([this](float DeltaTime) {
+
+			using namespace UE::Editor::DataStorage;
+			ICoreProvider* Storage = GetMutableDataStorageFeature<ICoreProvider>(StorageFeatureName);
+			
+			const UClass* LoadedClass = Cast<UClass>(StaticLoadObject(UClass::StaticClass(), nullptr, SelectedClassPath->ToString()));
+
+			for (RowHandle RowHandle : ReferenceQueryResultsNode->GetRows())
+			{
+				if (Storage->GetColumn<FTypedElementClassTypeInfoColumn>(RowHandle)->TypeInfo == LoadedClass)
+				{
+					ClassesTableViewer->SetSelection(RowHandle, true, ESelectInfo::Direct);
+					break;
+				}
+			}
+
+			return false;
+		}),
+		0.0f
+	);
+}
+
+TSharedRef<SWidget> SDataRecoveryTool::MakeTEDSTableViewer()
+{
+	using namespace UE::Editor::DataStorage;
+
+	ICoreProvider* Storage = UE::Editor::DataStorage::GetMutableDataStorageFeature<ICoreProvider>(StorageFeatureName);
+
+	if (!ensureMsgf(Storage, TEXT("Cannot create a TEDS ObjectReferenceContext widget before TEDS is initialized")))
+	{
+		return
+			SNew(STextBlock)
+			.Text(LOCTEXT("ObjectReferenceErrorCheck", "No valid editor data storage available."));
+	}
+
+	using namespace UE::Editor::DataStorage::Queries;
+
+	TSharedRef<QueryStack::FQueryNode> ReferenceQueryNode =
+		MakeShared<QueryStack::FQueryNode>(*Storage, Select().Where().All<FHasUnknownPropertiesTag, FTypedElementClassTypeInfoColumn>().Compile());
+
+	ReferenceQueryResultsNode = MakeShared<QueryStack::FRowQueryResultsNode>(*Storage,
+		ReferenceQueryNode,
+		QueryStack::FRowQueryResultsNode::ESyncFlags::RefreshOnUpdate);
+
+	TSharedPtr<QueryStack::IRowNode> SearchNode;
+	TSharedPtr<SBox> SearchBoxContainer = SNew(SBox);
+
+	SearchBoxContainer->SetContent(SNew(UE::Editor::DataStorage::STedsSearchBox)
+		.InSearchableRowNode(ReferenceQueryResultsNode)
+		.OutSearchNode(&SearchNode));
+
+
+	// If there is no SearchNode available, either disabled or there is an internal error, assign it to be the 
+	// ReferenceQueryResultsNode so that it won't be null
+	if (!SearchNode)
+	{
+		SearchNode = ReferenceQueryResultsNode;
+	}
+
+	// Using the Teds Outliner Purpose for matching customization for now
+	IUiProvider::FPurposeID PurposeId = IUiProvider::FPurposeInfo("SceneOutliner", "RowLabel", NAME_None).GeneratePurposeID();
+
+
+	ClassesTableViewer = SNew(STedsTableViewer)
+		.QueryStack(SearchNode)
+		.Columns({FTypedElementClassTypeInfoColumn::StaticStruct()})
+		.ListSelectionMode(ESelectionMode::Single)
+		.OnSelectionChanged(this, &SDataRecoveryTool::OnListSelectionChanged);
+
+	if (SelectedClassPath.IsValid())
+	{
+		SetSelectedClass();
+	}
+
+	return
+		SNew(SVerticalBox)
+		+ SVerticalBox::Slot()
+		.AutoHeight()
+		.Padding(Padding::Small)
+		[
+			SearchBoxContainer.ToSharedRef()
+		]
+		+ SVerticalBox::Slot()
+		[
+			ClassesTableViewer.ToSharedRef()
+		];
+}
+
+void SDataRecoveryTool::Tick(const FGeometry& AllottedGeometry, const double InCurrentTime, const float InDeltaTime)
+{
+}
+
+void SDataRecoveryTool::SetDockTab(const TSharedRef<SDockTab>& DockTab)
+{
+	OwningDockTab = DockTab;
+}
+
+FReply SDataRecoveryTool::OnConfirmClicked() const
+{
+	return FReply::Handled();
+}
+
+END_SLATE_FUNCTION_BUILD_OPTIMIZATION
+
+
+#undef LOCTEXT_NAMESPACE
+
diff --git a/Engine/Plugins/Experimental/InstanceDataObjectFixupTool/Source/Private/DataRecoveryToolStyleSheet.h b/Engine/Plugins/Experimental/InstanceDataObjectFixupTool/Source/Private/DataRecoveryToolStyleSheet.h
new file mode 100644
--- /dev/null
+++ b/Engine/Plugins/Experimental/InstanceDataObjectFixupTool/Source/Private/DataRecoveryToolStyleSheet.h
@@ -0,0 +1,61 @@
+// Copyright Epic Games, Inc. All Rights Reserved.
+
+#pragma once
+
+#include "CoreMinimal.h"
+#include "Styling/AppStyle.h"
+#include "Styling/SlateStyleRegistry.h"
+
+#define LOCTEXT_NAMESPACE "DataRecoveryToolModule"
+
+namespace UE::DataRecoveryTool::StyleSheet
+{
+	namespace Dialog
+	{
+		const FName DataRecoveryToolDialogName = FName("LoosePropertyFixup");
+	}
+
+	namespace Tab
+	{
+		const FName DataRecoveryToolTabName = FName("DataRecoveryTool");
+		const FName Icon = FName("Icons.Convert");
+		const FText DisplayName = LOCTEXT("DataRecoveryToolTabTitle", "Data Recovery");
+	}
+
+	namespace Padding
+	{
+		constexpr float None = 0.f;
+		constexpr float Small = 4.f;
+		constexpr float Normal = 8.f;
+		constexpr float Big = 12.f;
+	}
+
+	namespace Button
+	{
+		const FName Simple = FName("SimpleButton");
+	}
+
+	namespace Splitter
+	{
+		constexpr float MinSize = 290.f;
+	}
+
+	namespace Mappings::Column
+	{
+		const FName SourceProperty("SourceProperty");
+		const FName ArrowIcon("ArrowIcon");
+		const FName DestinationProperty("DestinationProperty");
+		const FName Alert("Alert");
+	}
+
+	inline const FSlateBrush* GetOpenExternalImage()
+	{
+		if (const ISlateStyle* RevisionControlStyle = FSlateStyleRegistry::FindSlateStyle("DefaultRevisionControlStyle"))
+		{
+			return RevisionControlStyle->GetBrush("RevisionControl.ConflictResolution.OpenExternal");
+		}
+		return FAppStyle::Get().GetBrush("Icons.CircleArrowUp");
+	}
+}
+
+#undef LOCTEXT_NAMESPACE
\ No newline at end of file
diff --git a/Engine/Plugins/Experimental/InstanceDataObjectFixupTool/Source/Private/DataRecoveryToolUtils.h b/Engine/Plugins/Experimental/InstanceDataObjectFixupTool/Source/Private/DataRecoveryToolUtils.h
new file mode 100644
--- /dev/null
+++ b/Engine/Plugins/Experimental/InstanceDataObjectFixupTool/Source/Private/DataRecoveryToolUtils.h
@@ -0,0 +1,21 @@
+// Copyright Epic Games, Inc. All Rights Reserved.
+
+#pragma once
+
+#include "CoreMinimal.h"
+#include "PropertyPath.h"
+#include "UObject/PropertyStateTracking.h"
+#include "UObject/PropertyPathName.h"
+#include "UObject/PropertyBagRepository.h"
+
+namespace UE::DataRecoveryTool::Utils
+{
+	inline TObjectPtr<UObject> GetInstanceDataObjectFromSelectedClassPath(const FTopLevelAssetPath& ClassPath)
+	{
+		const UClass* LoadedClass = Cast<UClass>(StaticLoadObject(UClass::StaticClass(), nullptr, ClassPath.ToString()));
+
+		const UObject* FoundInstance = FPropertyBagRepository::Get().FindFirstInstanceThatNeedsFixup(LoadedClass);
+
+		return FoundInstance ? FPropertyBagRepository::Get().FindInstanceDataObject(FoundInstance) : nullptr;
+	}
+}
diff --git a/Engine/Plugins/Experimental/InstanceDataObjectFixupTool/Source/Private/InstanceDataObjectFixupDetailCustomization.cpp b/Engine/Plugins/Experimental/InstanceDataObjectFixupTool/Source/Private/InstanceDataObjectFixupDetailCustomization.cpp
--- a/Engine/Plugins/Experimental/InstanceDataObjectFixupTool/Source/Private/InstanceDataObjectFixupDetailCustomization.cpp
+++ b/Engine/Plugins/Experimental/InstanceDataObjectFixupTool/Source/Private/InstanceDataObjectFixupDetailCustomization.cpp
@@ -246,9 +246,9 @@ TSharedRef<SWidget> FInstanceDataObjectNameWidgetOverride::GeneratePropertyRedir
 	MenuBuilder.EndSection();
 
 	TSet<FPropertyPath> RedirectOptions;
-	if (!Panel->Instances.IsEmpty())
+	if (Panel->InstanceDataObject)
 	{
-		UObject* FirstInstanceDataObject = Panel->Instances[0];
+		UObject* FirstInstanceDataObject = Panel->InstanceDataObject;
 		RedirectOptions = GetRedirectOptions(FirstInstanceDataObject->GetClass(), FirstInstanceDataObject);
 	}
 
diff --git a/Engine/Plugins/Experimental/InstanceDataObjectFixupTool/Source/Private/InstanceDataObjectFixupPanel.cpp b/Engine/Plugins/Experimental/InstanceDataObjectFixupTool/Source/Private/InstanceDataObjectFixupPanel.cpp
--- a/Engine/Plugins/Experimental/InstanceDataObjectFixupTool/Source/Private/InstanceDataObjectFixupPanel.cpp
+++ b/Engine/Plugins/Experimental/InstanceDataObjectFixupTool/Source/Private/InstanceDataObjectFixupPanel.cpp
@@ -71,9 +71,9 @@ static bool IsPathOverridden(UObject* Object, const FPropertyPath& Path)
 }
 
 FInstanceDataObjectFixupPanel::FInstanceDataObjectFixupPanel(
-	TConstArrayView<TObjectPtr<UObject>> InstanceDataObjects, TObjectPtr<UObject> InstanceDataObjectsOwner, EViewFlags InViewFlags)
-	: Instances(InstanceDataObjects)
-	, InstancesOwner(InstanceDataObjectsOwner)
+	const TObjectPtr<UObject>& InstanceDataObject, const TObjectPtr<UObject>& InstanceDataObjectOwner, EViewFlags InViewFlags)
+	: InstanceDataObject(InstanceDataObject)
+	, InstanceDataObjectOwner(InstanceDataObjectOwner)
 	, ViewFlags(InViewFlags)
 {
 }
@@ -88,26 +88,18 @@ FInstanceDataObjectFixupPanel::~FInstanceDataObjectFixupPanel()
 
 	if (DataStorageCompatibility != nullptr && DataStorage != nullptr)
 	{
-		for (UObject* Instance : Instances)
+		if (!Repository.RequiresFixup(Repository.FindInstanceForDataObject(InstanceDataObject)))
 		{
-			if (!Repository.RequiresFixup(Repository.FindInstanceForDataObject(Instance)))
-			{
-				static FName AlertName = FName("EntityLoosePropertiesErrorAlert");
+			static FName AlertName = FName("EntityLoosePropertiesErrorAlert");
 
-				// If a UObject isn't registered with TEDS, there's a chance its parent is registered and is the one
-				// with the alert column on it, so search upward until the nearest registered parent is found.
-				RowHandle Row = DataStorageCompatibility->FindRowWithCompatibleObject(InstancesOwner ? InstancesOwner.Get() : Instance);
-				Alerts::RemoveAlert(*DataStorage, Row, AlertName);
-			}
+			// If a UObject isn't registered with TEDS, there's a chance its parent is registered and is the one
+			// with the alert column on it, so search upward until the nearest registered parent is found.
+			RowHandle Row = DataStorageCompatibility->FindRowWithCompatibleObject(InstanceDataObjectOwner ? InstanceDataObjectOwner.Get() : InstanceDataObject.Get());
+			Alerts::RemoveAlert(*DataStorage, Row, AlertName);
 		}
 	}
 }
 
-int32 FInstanceDataObjectFixupPanel::Find(UObject* Value) const
-{
-	return Instances.Find(Value);
-}
-
 static bool RemoveCustomizationsWithLooseProperties(const FFieldVariant& FieldVariant, const TSharedPtr<IDetailsView>& DetailsView)
 {
 #if WITH_EDITORONLY_DATA
@@ -197,35 +189,25 @@ TSharedPtr<IDetailsView>& FInstanceDataObjectFixupPanel::GenerateDetailsView(boo
 	{
 		DetailsViewArgs.ShouldForceHideProperty.BindLambda([this](const TSharedRef<FPropertyNode>& PropertyNode)->bool
 		{
-			return !IsPathOverridden(Instances[0], *FPropertyNode::CreatePropertyPath(PropertyNode));
+			return !IsPathOverridden(InstanceDataObject, *FPropertyNode::CreatePropertyPath(PropertyNode));
 		});
 	}
 	
 	FPropertyEditorModule& PropertyEditorModule = FModuleManager::GetModuleChecked<FPropertyEditorModule>("PropertyEditor");
 	DetailsView = PropertyEditorModule.CreateDetailView(DetailsViewArgs);
-	for (const UObject* Instance : Instances)
-	{
-		RemoveCustomizationsWithLooseProperties(Instance->GetClass(), DetailsView);
-	}
+
+	RemoveCustomizationsWithLooseProperties(InstanceDataObject->GetClass(), DetailsView);
 	
-	DetailsView->SetObjects(Instances, true);
+	
+	DetailsView->SetObject(InstanceDataObject, true);
 
 	FCoreUObjectDelegates::OnObjectsReinstanced.AddSPLambda(this,
 		[this](const TMap<UObject*, UObject*>& ObjectMap)
 		{
-			bool bFound = false;
-			for (TObjectPtr<UObject>& Ido : Instances)
+			if (UObject* const* Found = ObjectMap.Find(InstanceDataObject))
 			{
-				if (UObject* const* Found = ObjectMap.Find(Ido))
-				{
-					Ido = TObjectPtr<UObject>(*Found);
-					bFound = true;
-				}
-			}
-
-			if (bFound)
-			{
-				DetailsView->SetObjects(Instances, true);
+				InstanceDataObject = TObjectPtr<UObject>(*Found);
+				DetailsView->SetObject(InstanceDataObject, true);
 			}
 		});
 	return DetailsView;
@@ -269,14 +251,8 @@ bool FInstanceDataObjectFixupPanel::ShouldSplitterIgnoreRow(const TWeakPtr<FDeta
 bool FInstanceDataObjectFixupPanel::AreAllConflictsRedirected() const
 {
 	UE::FPropertyBagRepository& Repository = UE::FPropertyBagRepository::Get();
-	for (UObject* Instance : Instances)
-	{
-		if (Repository.RequiresFixup(Repository.FindInstanceForDataObject(Instance)))
-		{
-			return false;
-		}
-	}
-	return true;
+
+	return !Repository.RequiresFixup(Repository.FindInstanceForDataObject(InstanceDataObject));
 }
 
 void FInstanceDataObjectFixupPanel::AutoApplyMarkDeletedActions()
@@ -462,7 +438,7 @@ void FInstanceDataObjectFixupPanel::RedirectProperty(const FPropertyPath& From,
 	// So that this IDT operation can be used on as many struct instances as possible, find the inner-most struct that's shared by both paths and chop them to be relative
 	// to that struct. So for example 'From' is "Foo (MyFooStruct) -> Bar (MyBarStruct) -> X (Float)" and 'To' is "Foo (MyFooStruct) -> Bar (MyBarStruct) -> Y (Float)"
 	// then we'll create an IDT for ALL instances of MyBarStruct that remaps "X (Float)" to "Y (Float)" instead of only remapping the full path.
-	const UObject* Owner = UE::FPropertyBagRepository::Get().FindInstanceForDataObject(Instances[0]);
+	const UObject* Owner = UE::FPropertyBagRepository::Get().FindInstanceForDataObject(InstanceDataObject);
 	UStruct* BaseStruct = Owner->GetClass();
 	int32 StructIndex = GetLastSharedStructSegment(FromPathNames[0], ToPathNames[0], BaseStruct);
 
@@ -487,7 +463,7 @@ void FInstanceDataObjectFixupPanel::RedirectProperty(const FPropertyPath& From,
 	TObjectPtr<UObject> OwnerObjectPtr = const_cast<UObject*>(Owner);
 
 	// this will likely trigger a IDO reinstance which will cause a details panel refresh for the new object!
-	UE::FInstanceDataTransforms::Get().PatchInstanceDataObject(Instances[0], OwnerObjectPtr);
+	UE::FInstanceDataTransforms::Get().PatchInstanceDataObject(InstanceDataObject, OwnerObjectPtr);
 }
 
 void FInstanceDataObjectFixupPanel::OnRedirectProperty(FPropertyPath From, FPropertyPath To)
@@ -502,7 +478,7 @@ bool FInstanceDataObjectFixupPanel::PropertyNeedsRedirect(const FPropertyPath& P
 		return false;
 	}
 
-	if (!IsPathOverridden(Instances[0], Path))
+	if (!IsPathOverridden(InstanceDataObject, Path))
 	{
 		return false;
 	}
@@ -515,7 +491,7 @@ bool FInstanceDataObjectFixupPanel::PropertyNeedsRedirect(const FPropertyPath& P
 		return false;
 	}
 
-	const UObject* Owner = UE::FPropertyBagRepository::Get().FindInstanceForDataObject(Instances[0]);
+	const UObject* Owner = UE::FPropertyBagRepository::Get().FindInstanceForDataObject(InstanceDataObject);
 	return UE::FPropertyBagRepository::Get().RequiresFixup(Owner, PathNames[0]);
 }
 
@@ -539,12 +515,12 @@ void FInstanceDataObjectFixupPanel::MarkForDelete(const FPropertyPath& CurrentPa
 		ensureMsgf(false, TEXT("out of subobjects is not supported yet!"));
 		return;
 	}
-	const UObject* Owner = UE::FPropertyBagRepository::Get().FindInstanceForDataObject(Instances[0]);
+	const UObject* Owner = UE::FPropertyBagRepository::Get().FindInstanceForDataObject(InstanceDataObject);
 	TSharedPtr<FJsonObject> Params = MakeShared<FJsonObject>();
 	Params->SetField(TEXT("src"), srcPathNames[0].ToJsonValue());
 	UE::FInstanceDataTransforms::Get().WriteInstanceDataTransform(*Owner->GetClass(), TEXT("RemoveProperty"), Params);
 	TObjectPtr<UObject> OwnerObjectPtr = const_cast<UObject*>(Owner);
-	UE::FInstanceDataTransforms::Get().PatchInstanceDataObject(Instances[0], OwnerObjectPtr);
+	UE::FInstanceDataTransforms::Get().PatchInstanceDataObject(InstanceDataObject, OwnerObjectPtr);
 }
 
 void FInstanceDataObjectFixupPanel::OnMarkForDelete(FPropertyPath Path)
diff --git a/Engine/Plugins/Experimental/InstanceDataObjectFixupTool/Source/Private/InstanceDataObjectFixupPanel.h b/Engine/Plugins/Experimental/InstanceDataObjectFixupTool/Source/Private/InstanceDataObjectFixupPanel.h
--- a/Engine/Plugins/Experimental/InstanceDataObjectFixupTool/Source/Private/InstanceDataObjectFixupPanel.h
+++ b/Engine/Plugins/Experimental/InstanceDataObjectFixupTool/Source/Private/InstanceDataObjectFixupPanel.h
@@ -17,8 +17,8 @@ public:
 		None = 0,
 		HideLooseProperties = (1 << 0), // hide properties with isLoose metadata set to true
 		IncludeOnlySetBySerialization = (1 << 1), // hide properties that weren't set by serialization
-		AllowRemapLooseProperties = (1 << 3),
 		ReadonlyValues = (1 << 2),
+		AllowRemapLooseProperties = (1 << 3),
 
 		// displays only properties found in the property bag, allow remapping, and disallow value edits
 		DefaultLeftPanel = IncludeOnlySetBySerialization | AllowRemapLooseProperties | ReadonlyValues,
@@ -27,10 +27,9 @@ public:
 	};
 	
 	FInstanceDataObjectFixupPanel(
-		TConstArrayView<TObjectPtr<UObject>> InstanceDataObjects, TObjectPtr<UObject> InstanceDataObjectsOwner, EViewFlags ViewFlags);
+		const TObjectPtr<UObject>& InstanceDataObject, const TObjectPtr<UObject>& InstanceDataObjectOwner, EViewFlags ViewFlags);
 	~FInstanceDataObjectFixupPanel();
 	
-	int32 Find(UObject* Value) const;
 	TSharedPtr<IDetailsView>& GenerateDetailsView(bool bScrollbarOnLeft = false);
 	
 	void SetDiffAgainstLeft(const TSharedPtr<FAsyncDetailViewDiff>& DiffAgainstLeft);
@@ -72,9 +71,9 @@ private:
 	void RedirectProperty(const FPropertyPath& From, const FPropertyPath& To);
 	// pass-by-copy version for delegates. Use RedirectProperty when possible
 	void OnRedirectProperty(FPropertyPath From, FPropertyPath To);
-	
-	TArray<TObjectPtr<UObject>> Instances; // stores either InstanceDataObject property bag or just regular objects
-	TObjectPtr<UObject> InstancesOwner; // if not null, then it stores the common parent for the instances.
+
+	TObjectPtr<UObject> InstanceDataObject; // stores either InstanceDataObject property bag or just regular objects
+	TObjectPtr<UObject> InstanceDataObjectOwner; // if not null, then it stores the common parent for the instances.
 
 	TMap<FPropertyPath, FRevertInfo> RevertInfo;
 	TSet<FPropertyPath> MarkedForDelete;
diff --git a/Engine/Plugins/Experimental/InstanceDataObjectFixupTool/Source/Private/InstanceDataObjectFixupTool.cpp b/Engine/Plugins/Experimental/InstanceDataObjectFixupTool/Source/Private/InstanceDataObjectFixupTool.cpp
--- a/Engine/Plugins/Experimental/InstanceDataObjectFixupTool/Source/Private/InstanceDataObjectFixupTool.cpp
+++ b/Engine/Plugins/Experimental/InstanceDataObjectFixupTool/Source/Private/InstanceDataObjectFixupTool.cpp
@@ -9,10 +9,14 @@
 #include "SDetailsSplitter.h"
 #include "Widgets/Layout/LinkableScrollBar.h"
 #include "Widgets/Input/SButton.h"
-
+#include "DataRecoveryToolStyleSheet.h"
+#include "DataRecoveryToolUtils.h"
+#include "Widgets/Layout/SScrollBox.h"
 
 #define LOCTEXT_NAMESPACE "InstanceDataObjectFixupTool"
 
+using namespace UE::DataRecoveryTool::StyleSheet;
+
 class FInstanceDataObjectFixupSpecification : public TTreeDiffSpecification<TWeakPtr<FDetailTreeNode>>
 {
 public:
@@ -67,47 +71,68 @@ public:
 BEGIN_SLATE_FUNCTION_BUILD_OPTIMIZATION
 void SInstanceDataObjectFixupTool::Construct(const FArguments& InArgs)
 {
-	Panels[0] = MakeShared<FInstanceDataObjectFixupPanel>(InArgs._InstanceDataObjects, InArgs._InstanceDataObjectsOwner, FInstanceDataObjectFixupPanel::EViewFlags::DefaultLeftPanel);
-	Panels[1] = MakeShared<FInstanceDataObjectFixupPanel>(InArgs._InstanceDataObjects, InArgs._InstanceDataObjectsOwner, FInstanceDataObjectFixupPanel::EViewFlags::DefaultRightPanel);
-	
+	SelectedClassPath = InArgs._SelectedClassPath;
+
 	ChildSlot
 	[
 		SNew(SVerticalBox)
-		+SVerticalBox::Slot()
-		.AutoHeight()
-		.VAlign(VAlign_Top)
+		+ SVerticalBox::Slot()
+		.Padding(Padding::None, Padding::Small)
+		.FillHeight(1.f)
 		[
-			SNew(SHorizontalBox)
-			+SHorizontalBox::Slot()
-			.AutoWidth()
-			.HAlign(HAlign_Left)
+			SNew(SScrollBox)
+			.Orientation(EOrientation::Orient_Vertical)
+			+SScrollBox::Slot()
 			[
-				SNew(SButton)
-				.Text(LOCTEXT("AutoMarkForDeletion", "Mark remaining conflicts for deletion"))
-				.OnClicked(this, &SInstanceDataObjectFixupTool::OnAutoMarkForDeletion)
+				SAssignNew(Border, SBorder)
+				.BorderImage(FAppStyle::Get().GetBrush("Brushes.Recessed"))
+				.Padding(Padding::None)
 			]
 		]
 		+SVerticalBox::Slot()
-		[
-			SAssignNew(Splitter, SDetailsSplitter)
-			.RowHighlightColor_Static(&SInstanceDataObjectFixupTool::GetRowHighlightColor)
-		]
-		+SVerticalBox::Slot()
-		.HAlign(HAlign_Right)
-		.VAlign(VAlign_Center)
 		.AutoHeight()
+		.HAlign(HAlign_Left)
+		.VAlign(VAlign_Center)
+		.Padding(Padding::Small)
 		[
 			SNew(SButton)
-			.Text(LOCTEXT("ConfirmButton","Confirm"))
-			.OnClicked(this, &SInstanceDataObjectFixupTool::OnConfirmClicked)
-			.IsEnabled(this, &SInstanceDataObjectFixupTool::IsResolved)
+			.HAlign(HAlign_Center)
+			.VAlign(VAlign_Center)
+			.ContentPadding(Padding::Normal)
+			.ButtonStyle(FAppStyle::Get(), "SimpleButton")
+			.ToolTipText(LOCTEXT("AutoMarkForDeletion_Tooltip", "Mark remaining conflicts for deletion"))
+			.OnClicked(this, &SInstanceDataObjectFixupTool::OnAutoMarkForDeletion)
+			[
+				SNew(SHorizontalBox)
+				+SHorizontalBox::Slot()
+				.AutoWidth()
+				.VAlign(VAlign_Center)
+				[
+					SNew(SImage)
+					.Image(FAppStyle::GetBrush("Icons.Delete"))
+				]
+				+ SHorizontalBox::Slot()
+				.AutoWidth()
+				.Padding(Padding::Small, Padding::Small, Padding::Small, Padding::Small)
+				.VAlign(VAlign_Center)
+				[
+					SNew(STextBlock)
+					.Text(LOCTEXT("AutoMarkForDeletion", "Delete remaining conflicts"))
+				]
+			]
 		]
 	];
+
+	// TODO: Could remove this for optimizations
+	GenerateDetailsViews();
 }
 
 void SInstanceDataObjectFixupTool::Tick(const FGeometry& AllottedGeometry, const double InCurrentTime, const float InDeltaTime)
 {
-	PanelDiff->Tick();
+	if (PanelDiff != nullptr)
+	{ 
+		PanelDiff->Tick();
+	}
 }
 
 END_SLATE_FUNCTION_BUILD_OPTIMIZATION
@@ -126,11 +151,36 @@ FReply SInstanceDataObjectFixupTool::OnConfirmClicked() const
 	return FReply::Handled();
 }
 
+void SInstanceDataObjectFixupTool::ResetDiff()
+{
+	Border->SetContent(SNullWidget::NullWidget);
+	PanelDiff.Reset();
+}
+
 void SInstanceDataObjectFixupTool::GenerateDetailsViews()
 {
+	TSharedPtr<FTopLevelAssetPath> ClassPath = SelectedClassPath.Pin();
+
+	if (!ClassPath)
+	{
+		ResetDiff();
+		return;
+	}
+
+	TObjectPtr<UObject> InstanceDataObject = UE::DataRecoveryTool::Utils::GetInstanceDataObjectFromSelectedClassPath(*ClassPath);
+
+	if (InstanceDataObject == nullptr)
+	{
+		ResetDiff();
+		return;
+	}
+
 	constexpr int32 LeftIndex = 0;
 	constexpr int32 RightIndex = 1;
-	
+
+	Panels[LeftIndex] = MakeShared<FInstanceDataObjectFixupPanel>(InstanceDataObject, nullptr, FInstanceDataObjectFixupPanel::EViewFlags::DefaultLeftPanel);
+	Panels[RightIndex] = MakeShared<FInstanceDataObjectFixupPanel>(InstanceDataObject, nullptr, FInstanceDataObjectFixupPanel::EViewFlags::DefaultRightPanel);
+
 	// generate the panels
 	Panels[LeftIndex]->GenerateDetailsView(true);
 	Panels[RightIndex]->GenerateDetailsView(false);
@@ -139,25 +189,35 @@ void SInstanceDataObjectFixupTool::GenerateDetailsViews()
 	PanelDiff = MakeShared<FAsyncDetailViewDiff>(
 		Panels[LeftIndex]->DetailsView.ToSharedRef(), Panels[RightIndex]->DetailsView.ToSharedRef());
 	PanelDiff->SetDiffSpecification<FInstanceDataObjectFixupSpecification>(Panels[LeftIndex], Panels[RightIndex]);
-    
+
 	SLinkableScrollBar::LinkScrollBars(Panels[LeftIndex]->LinkableScrollBar.ToSharedRef(), Panels[RightIndex]->LinkableScrollBar.ToSharedRef(),
 		TAttribute<TArray<FVector2f>>::CreateRaw(PanelDiff.Get(), &FAsyncDetailViewDiff::GenerateScrollSyncRate)); // TODO: Make work w/ CreateShared
 	Panels[LeftIndex]->SetDiffAgainstRight(PanelDiff);
 	Panels[RightIndex]->SetDiffAgainstLeft(PanelDiff);
 
+
+	Border->SetContent(
+		SAssignNew(Splitter, SDetailsSplitter)
+		.RowHighlightColor_Static(&SInstanceDataObjectFixupTool::GetRowHighlightColor)
+	);
+
 	// add the panels to the splitter
-	for (TSharedPtr<FInstanceDataObjectFixupPanel> Panel : Panels)
+	for (TSharedPtr Panel : Panels)
 	{
 		Splitter->AddSlot(
 			SDetailsSplitter::Slot()
 			.DetailsView(Panel->DetailsView.ToSharedRef())
 			.DifferencesWithRightPanel(Panel.Get(), &FInstanceDataObjectFixupPanel::GetDiffAgainstRight)
-			.IsReadonly_Lambda([Panel = Panel.ToWeakPtr()](){return Panel.IsValid() ? Panel.Pin()->HasViewFlag(FInstanceDataObjectFixupPanel::EViewFlags::ReadonlyValues) : true;})
+			.IsReadonly_Lambda([Panel = Panel.ToWeakPtr()]() { return Panel.IsValid() ? 
+				Panel.Pin()->HasViewFlag(FInstanceDataObjectFixupPanel::EViewFlags::ReadonlyValues) : true; })
 			.ShouldIgnoreRow(Panel.Get(), &FInstanceDataObjectFixupPanel::ShouldSplitterIgnoreRow)
 		);
 	}
+	
 }
 
+	
+
 FLinearColor SInstanceDataObjectFixupTool::GetRowHighlightColor(const TUniquePtr<FAsyncDetailViewDiff::DiffNodeType>& DiffNode)
 {
 	switch (DiffNode->DiffResult)
diff --git a/Engine/Plugins/Experimental/InstanceDataObjectFixupTool/Source/Private/InstanceDataObjectFixupToolModule.cpp b/Engine/Plugins/Experimental/InstanceDataObjectFixupTool/Source/Private/InstanceDataObjectFixupToolModule.cpp
--- a/Engine/Plugins/Experimental/InstanceDataObjectFixupTool/Source/Private/InstanceDataObjectFixupToolModule.cpp
+++ b/Engine/Plugins/Experimental/InstanceDataObjectFixupTool/Source/Private/InstanceDataObjectFixupToolModule.cpp
@@ -1,95 +1,66 @@
 // Copyright Epic Games, Inc. All Rights Reserved.
 
 #include "InstanceDataObjectFixupToolModule.h"
+
+#include "DataRecoveryTool.h"
+#include "DataRecoveryToolStyleSheet.h"
 #include "InstanceDataObjectFixupTool.h"
 #include "WorkspaceMenuStructure.h"
 #include "WorkspaceMenuStructureModule.h"
+#include "Elements/Framework/TypedElementRegistry.h"
+#include "Framework/Application/SlateApplication.h"
+#include "Modules/ModuleManager.h"
+#include "Widgets/Docking/SDockTab.h"
 
 #define LOCTEXT_NAMESPACE "InstanceDataObjectFixupToolModule"
 
-static const FName InstanceDataObjectFixupToolTabName = FName(TEXT("InstanceDataObjectFixupTool"));
-static const FName InstanceDataObjectFixupToolDialogName = FName(TEXT("LoosePropertyFixup"));
-
 void FInstanceDataObjectFixupToolModule::StartupModule()
 {
 }
 
 void FInstanceDataObjectFixupToolModule::ShutdownModule()
 {
-	
-}
+	using namespace UE::DataRecoveryTool::StyleSheet;
 
-bool FInstanceDataObjectFixupToolModule::OpenInstanceDataObjectFixupTool() const
-{
-	TSharedPtr<SDockTab> DockTab = FGlobalTabmanager::Get()->TryInvokeTab(FTabId(InstanceDataObjectFixupToolTabName));
-	if (!DockTab)
+	if (FSlateApplication::IsInitialized())
 	{
-		return false;
+		FGlobalTabmanager::Get()->UnregisterNomadTabSpawner(Dialog::DataRecoveryToolDialogName);
 	}
-	DockTab->DrawAttention();
-	return true;
 }
 
-TSharedRef<SDockTab> FInstanceDataObjectFixupToolModule::CreateInstanceDataObjectFixupTab(
-	const FSpawnTabArgs& TabArgs, TConstArrayView<TObjectPtr<UObject>> InstanceDataObjects, TObjectPtr<UObject> InstanceDataObjectsOwner) const
+TSharedRef<SDockTab> FInstanceDataObjectFixupToolModule::CreateDataRecoveryToolTab(TOptional<FTopLevelAssetPath> ClassPath) const
 {
-	const TSharedRef<SInstanceDataObjectFixupTool> InstanceDataObjectFixupTool = SNew(SInstanceDataObjectFixupTool)
-		.InstanceDataObjects(InstanceDataObjects)
-		.InstanceDataObjectsOwner(InstanceDataObjectsOwner);
+	const TSharedRef<SDataRecoveryTool> DataRecoveryTool = SNew(SDataRecoveryTool)
+		.SelectedClassPath(ClassPath);
+
 	TSharedRef<SDockTab> DockTab = SNew(SDockTab)
 		.TabRole(ETabRole::NomadTab)
 		[
-			InstanceDataObjectFixupTool
+			DataRecoveryTool
 		];
 
-	InstanceDataObjectFixupTool->SetDockTab(DockTab);
-	InstanceDataObjectFixupTool->GenerateDetailsViews();
-	
+	DataRecoveryTool->SetDockTab(DockTab);
+
 	return DockTab;
 }
 
-void FInstanceDataObjectFixupToolModule::CreateInstanceDataObjectFixupDialog(
-	TConstArrayView<TObjectPtr<UObject>> InstanceDataObjects, TObjectPtr<UObject> InstanceDataObjectsOwner) const
+void FInstanceDataObjectFixupToolModule::CreateDataRecoveryToolDialog(const UObject* Instance, const UObject* /*Owner = nullptr*/) const
 {
-	FText DisplayName;
-	if (InstanceDataObjects.Num() == 1)
-	{
-		DisplayName = FText::Format(LOCTEXT("SingleObjectFixupTabTitle", "{0} Fix-up"), FText::FromName(InstanceDataObjects[0]->GetFName()));
-	}
-	else
-	{
-		FString ClassName;
-		for (UObject* Object : InstanceDataObjects)
-		{
-			if (ClassName.IsEmpty())
-			{
-				ClassName = Object->GetClass()->GetName();
-			}
-			else if (ClassName != Object->GetName())
-			{
-				ClassName.Empty();
-				break;
-			}
-		}
-		if (ClassName.IsEmpty())
-		{
-			DisplayName = FText::Format(LOCTEXT("MultiEditObjectFixupTabTitle", "{0} Objects Fixup"), InstanceDataObjects.Num());
-		}
-		else
-		{
-			DisplayName = FText::Format(LOCTEXT("MultiEditObjectFixupTabTitleWithClass", "{0} {1} Objects Fixup"), InstanceDataObjects.Num(), FText::FromString(ClassName));
-		}
-	}
-	
-	FGlobalTabmanager::Get()->RegisterNomadTabSpawner(InstanceDataObjectFixupToolDialogName, FOnSpawnTab::CreateLambda(
-		[&InstanceDataObjects, InstanceDataObjectsOwner](const FSpawnTabArgs& TabArgs)
+	using namespace UE::DataRecoveryTool::StyleSheet;
+
+	FGlobalTabmanager::Get()->RegisterNomadTabSpawner(Dialog::DataRecoveryToolDialogName, FOnSpawnTab::CreateLambda(
+		[Instance](const FSpawnTabArgs& TabArgs)
 		{
-			return FInstanceDataObjectFixupToolModule::Get().CreateInstanceDataObjectFixupTab(TabArgs, InstanceDataObjects, InstanceDataObjectsOwner);
+			TOptional<FTopLevelAssetPath> ClassPath = Instance ? Instance->GetClass()->GetClassPathName() : TOptional<FTopLevelAssetPath>{};
+
+			return Get().CreateDataRecoveryToolTab(ClassPath);
 		}))
-		.SetDisplayName(DisplayName)
+		.SetIcon(FSlateIcon(FAppStyle::GetAppStyleSetName(), Tab::Icon))
+		.SetDisplayName(Tab::DisplayName)
 		.SetGroup(WorkspaceMenu::GetMenuStructure().GetToolsCategory())
 		.SetMenuType(ETabSpawnerMenuType::Hidden);
-	if (const TSharedPtr<SDockTab> DockTab = FGlobalTabmanager::Get()->TryInvokeTab(FTabId(InstanceDataObjectFixupToolDialogName)))
+
+	if (const TSharedPtr<SDockTab> DockTab = FGlobalTabmanager::Get()->TryInvokeTab(FTabId(Dialog::DataRecoveryToolDialogName)))
 	{
 		DockTab->DrawAttention();
 	}
diff --git a/Engine/Plugins/Experimental/InstanceDataObjectFixupTool/Source/Private/InstanceDataObjectFixupToolTedsQueries.cpp b/Engine/Plugins/Experimental/InstanceDataObjectFixupTool/Source/Private/InstanceDataObjectFixupToolTedsQueries.cpp
--- a/Engine/Plugins/Experimental/InstanceDataObjectFixupTool/Source/Private/InstanceDataObjectFixupToolTedsQueries.cpp
+++ b/Engine/Plugins/Experimental/InstanceDataObjectFixupTool/Source/Private/InstanceDataObjectFixupToolTedsQueries.cpp
@@ -74,14 +74,14 @@ void UInstanceDataObjectFixupToolTedsQueryFactory::ShowFixUpTool(UE::Editor::Dat
 			UE::FPropertyBagRepository::Get().FindNestedInstanceDataObject(Owner, true,
 				[Owner](UObject* NestedObject)
 				{
-					FInstanceDataObjectFixupToolModule::Get().CreateInstanceDataObjectFixupDialog({NestedObject}, Owner);
+					FInstanceDataObjectFixupToolModule::Get().CreateDataRecoveryToolDialog(NestedObject, Owner);
 				});
 		}
 		else
 		{
 			if (TObjectPtr<UObject> InstanceDataObject = UE::FPropertyBagRepository::Get().FindInstanceDataObject(ObjectColumn->Object.Get()))
 			{
-				FInstanceDataObjectFixupToolModule::Get().CreateInstanceDataObjectFixupDialog({ InstanceDataObject });
+				FInstanceDataObjectFixupToolModule::Get().CreateDataRecoveryToolDialog(InstanceDataObject);
 			}
 		}
 	}
diff --git a/Engine/Plugins/Experimental/InstanceDataObjectFixupTool/Source/Public/DataRecoveryTool.h b/Engine/Plugins/Experimental/InstanceDataObjectFixupTool/Source/Public/DataRecoveryTool.h
new file mode 100644
--- /dev/null
+++ b/Engine/Plugins/Experimental/InstanceDataObjectFixupTool/Source/Public/DataRecoveryTool.h
@@ -0,0 +1,68 @@
+// Copyright Epic Games, Inc. All Rights Reserved.
+
+#pragma once
+
+#include "CoreMinimal.h"
+#include "TedsTableViewer/Public/TedsTableViewerModel.h"
+#include "TedsTableViewer/Public/Widgets/STedsTableViewer.h"
+#include "Widgets/Docking/SDockTab.h"
+#include "Widgets/SCompoundWidget.h"
+#include "DataRecoveryToolUtils.h"
+#include "UObject/InstanceDataTransforms.h"
+
+#define UE_API INSTANCEDATAOBJECTFIXUPTOOL_API
+
+namespace UE::Editor::DataStorage::QueryStack
+{
+	class FRowQueryResultsNode;
+}
+
+class SInstanceDataObjectFixupTool;
+
+class SDataRecoveryTool : public SCompoundWidget
+{
+public:
+	SLATE_BEGIN_ARGS(SDataRecoveryTool)
+	{}
+		SLATE_ARGUMENT(TOptional<FTopLevelAssetPath>, SelectedClassPath)
+	SLATE_END_ARGS()
+
+	/** Constructs this widget with InArgs */
+	UE_API void Construct(const FArguments& InArgs);
+	UE_API virtual void Tick(const FGeometry& AllottedGeometry, const double InCurrentTime, const float InDeltaTime) override;
+
+	UE_API void SetDockTab(const TSharedRef<SDockTab>& DockTab);
+
+private:
+
+	UE_API TSharedRef<SWidget> MakeTopBar();
+
+	UE_API TSharedRef<SWidget> MakePanel(const FText& InTitle, const TSharedRef<SWidget>& InContent,
+		const TOptional<FOnGetContent>& InMenuContent = TOptional<FOnGetContent>{});
+
+	UE_API TSharedRef<SWidget> MakeClassesPanel();
+	UE_API TSharedRef<SWidget> MakeTEDSTableViewer();
+	UE_API void OnListSelectionChanged(UE::Editor::DataStorage::RowHandle RowItem);
+	void SetSelectedClass();
+
+	UE_API TSharedRef<SWidget> MakeRightPanel();
+
+	UE_API FReply OnSaveClicked();
+
+	UE_API TSharedRef<SWidget> MakeDiffSection();
+
+	UE_API FReply OnConfirmClicked() const;
+
+
+	TWeakPtr<SDockTab> OwningDockTab;
+
+	TSharedPtr<FTopLevelAssetPath> SelectedClassPath;
+
+	TSharedPtr<UE::Editor::DataStorage::STedsTableViewer> ClassesTableViewer;
+
+	TSharedPtr<UE::Editor::DataStorage::QueryStack::FRowQueryResultsNode> ReferenceQueryResultsNode;
+
+	TSharedPtr<SInstanceDataObjectFixupTool> DiffView;
+};
+
+#undef UE_API
diff --git a/Engine/Plugins/Experimental/InstanceDataObjectFixupTool/Source/Public/InstanceDataObjectFixupTool.h b/Engine/Plugins/Experimental/InstanceDataObjectFixupTool/Source/Public/InstanceDataObjectFixupTool.h
--- a/Engine/Plugins/Experimental/InstanceDataObjectFixupTool/Source/Public/InstanceDataObjectFixupTool.h
+++ b/Engine/Plugins/Experimental/InstanceDataObjectFixupTool/Source/Public/InstanceDataObjectFixupTool.h
@@ -4,6 +4,7 @@
 
 #include "CoreMinimal.h"
 #include "AsyncDetailViewDiff.h"
+#include "UObject/InstanceDataTransforms.h"
 #include "Widgets/SCompoundWidget.h"
 #include "Widgets/Docking/SDockTab.h"
 
@@ -23,8 +24,8 @@ class SInstanceDataObjectFixupTool : public SCompoundWidget
 public:
 	SLATE_BEGIN_ARGS(SInstanceDataObjectFixupTool)
 	{}
-		SLATE_ARGUMENT(TConstArrayView<TObjectPtr<UObject>>, InstanceDataObjects)
-		SLATE_ARGUMENT(TObjectPtr<UObject>, InstanceDataObjectsOwner)
+		SLATE_ARGUMENT(TWeakPtr<FTopLevelAssetPath>, SelectedClassPath)
+
 	SLATE_END_ARGS()
 
 	/** Constructs this widget with InArgs */
@@ -39,12 +40,16 @@ public:
 	
 private:
 	UE_API FReply OnConfirmClicked() const;
-	
+	void ResetDiff();
+
 	TWeakPtr<SDockTab> OwningDockTab;
 	TStaticArray<TSharedPtr<FInstanceDataObjectFixupPanel>, 2> Panels;
 	TSharedPtr<FAsyncDetailViewDiff> PanelDiff;
 
 	TSharedPtr<SDetailsSplitter> Splitter;
+
+	TSharedPtr<SBorder> Border;
+	TWeakPtr<FTopLevelAssetPath> SelectedClassPath;
 };
 
 #undef UE_API
diff --git a/Engine/Plugins/Experimental/InstanceDataObjectFixupTool/Source/Public/InstanceDataObjectFixupToolModule.h b/Engine/Plugins/Experimental/InstanceDataObjectFixupTool/Source/Public/InstanceDataObjectFixupToolModule.h
--- a/Engine/Plugins/Experimental/InstanceDataObjectFixupTool/Source/Public/InstanceDataObjectFixupToolModule.h
+++ b/Engine/Plugins/Experimental/InstanceDataObjectFixupTool/Source/Public/InstanceDataObjectFixupToolModule.h
@@ -11,21 +11,15 @@ class FInstanceDataObjectFixupToolModule : public FDefaultModuleImpl
 public:
 	UE_API virtual void StartupModule() override;
 	UE_API virtual void ShutdownModule() override;
-	UE_API bool OpenInstanceDataObjectFixupTool() const;
+
 	static FInstanceDataObjectFixupToolModule& Get()
 	{
 		return FModuleManager::LoadModuleChecked<FInstanceDataObjectFixupToolModule>("InstanceDataObjectFixupTool");
 	}
-
-	// opens the fixup tool 
-	UE_API TSharedRef<SDockTab> CreateInstanceDataObjectFixupTab(
-		const FSpawnTabArgs& TabArgs, 
-		TConstArrayView<TObjectPtr<UObject>> InstanceDataObjects, 
-		TObjectPtr<UObject> InstanceDataObjectsOwner = nullptr) const;
 	
-	UE_API void CreateInstanceDataObjectFixupDialog(
-		TConstArrayView<TObjectPtr<UObject>> InstanceDataObjects, 
-		TObjectPtr<UObject> InstanceDataObjectsOwner = nullptr) const;
+	UE_API TSharedRef<SDockTab> CreateDataRecoveryToolTab(TOptional<FTopLevelAssetPath> ClassPath = {}) const;
+
+	UE_API void CreateDataRecoveryToolDialog(const UObject* Instance, const UObject* Owner = nullptr) const;
 };
 
 #undef UE_API
