47175615 was missing these two files by mistake.
[Feature] Updated IDT api to allow users to get, preview, and edit the full set of transforms rather than simply appending to the list and saving immediately

diff --git a/Engine/Plugins/Experimental/InstanceDataObjectFixupTool/Source/Private/DataRecoveryTool.cpp b/Engine/Plugins/Experimental/InstanceDataObjectFixupTool/Source/Private/DataRecoveryTool.cpp
--- a/Engine/Plugins/Experimental/InstanceDataObjectFixupTool/Source/Private/DataRecoveryTool.cpp
+++ b/Engine/Plugins/Experimental/InstanceDataObjectFixupTool/Source/Private/DataRecoveryTool.cpp
@@ -21,10 +21,12 @@
 #include "TedsQueryStack/Public/TedsRowQueryResultsNode.h"
 #include "TedsTableViewer/Public/Widgets/STedsSearchBox.h"
 #include "TedsTableViewer/Public/Widgets/STedsTableViewer.h"
+#include "UObject/InstanceDataTransforms.h"
 #include "Widgets/Input/SButton.h"
 #include "Widgets/Layout/SSeparator.h"
 #include "Widgets/STedsHierarchyViewer.h"
 #include "Serialization/ObjectWriter.h"
+#include "Misc/ScopedSlowTask.h"
 
 
 #define LOCTEXT_NAMESPACE "DataRecoveryTool"
@@ -36,6 +38,7 @@ using namespace UE::DataRecoveryTool::StyleSheet;
 void SDataRecoveryTool::Construct(const FArguments& InArgs)
 {
 	SelectedClassPath =  MakeShared<FTopLevelAssetPath>(InArgs._SelectedClassPath.IsSet() ? *InArgs._SelectedClassPath : nullptr);
+	StagedTransforms = MakeShared<TMap<FTopLevelAssetPath, UE::FInstanceDataTransformSet>>();
 	ChildSlot
 		[
 			SNew(SVerticalBox)
@@ -62,7 +65,6 @@ void SDataRecoveryTool::Construct(const FArguments& InArgs)
 				]
 			]
 		];
-
 }
 
 
@@ -84,7 +86,7 @@ TSharedRef<SWidget> SDataRecoveryTool::MakeTopBar()
 				.VAlign(VAlign_Center)
 				.ContentPadding(Padding::Normal)
 				.ButtonStyle(FAppStyle::Get(), "SimpleButton")
-				.OnClicked(this, &SDataRecoveryTool::OnSaveClicked)
+				.OnClicked(this, &SDataRecoveryTool::OnSaveAndApplyClicked)
 				[
 					SNew(SImage)
 					.Image(FAppStyle::GetBrush("Icons.Save"))
@@ -113,6 +115,7 @@ TSharedRef<SWidget> SDataRecoveryTool::MakeTopBar()
 						.VAlign(VAlign_Center)
 						.ContentPadding(Padding::Normal)
 						.ButtonStyle(FAppStyle::Get(), "SimpleButton")
+						.OnClicked(this, &SDataRecoveryTool::OnApplyClicked)
 						[
 							SNew(SHorizontalBox)
 							+ SHorizontalBox::Slot()
@@ -177,13 +180,6 @@ TSharedRef<SWidget> SDataRecoveryTool::MakeRightPanel()
 	return MakePanel(LOCTEXT("MappingsTitle", "MAPPINGS"), MakeDiffSection());
 }
 
-FReply SDataRecoveryTool::OnSaveClicked()
-{
-	//TODO: Implement this
-	return FReply::Handled();
-}
-
-
 TSharedRef<SWidget> SDataRecoveryTool::MakeDiffSection()
 {
 	return
@@ -193,9 +189,106 @@ TSharedRef<SWidget> SDataRecoveryTool::MakeDiffSection()
 		[
 			SAssignNew(DiffView, SInstanceDataObjectFixupTool)
 				.SelectedClassPath(SelectedClassPath)
+				.StagedTransforms(StagedTransforms)
 		];
 }
 
+FReply SDataRecoveryTool::OnApplyClicked() const
+{
+	using namespace UE::Editor::DataStorage::Queries;
+	
+	// for now just apply the transforms to all instances of every listed class. This may be slow!
+	ICoreProvider* Storage = UE::Editor::DataStorage::GetMutableDataStorageFeature<ICoreProvider>(StorageFeatureName);
+	UE::Editor::DataStorage::QueryHandle Query = Storage->RegisterQuery(
+		Select()
+		.ReadOnly<FTypedElementClassTypeInfoColumn>()
+		.Where()
+		.All<FHasUnknownPropertiesTag>()
+		.Compile());
+
+	TArray<UObject*> AllInstancesOfClasses;
+	Storage->RunQuery(
+		Query,
+		CreateDirectQueryCallbackBinding(
+		[&AllInstancesOfClasses](IDirectQueryContext&, const FTypedElementClassTypeInfoColumn& TypeInfo)
+		{
+			if (const UClass* Class = TypeInfo.TypeInfo.Get())
+			{
+				GetObjectsOfClass(Class, AllInstancesOfClasses);
+			}
+		})
+	);
+
+	FScopedSlowTask ApplyAllTask{ static_cast<float>(AllInstancesOfClasses.Num()), LOCTEXT("ApplyAllSlow","Applying transforms") };
+	ApplyAllTask.MakeDialog();
+
+	UE::FPropertyBagRepository& Repository = UE::FPropertyBagRepository::Get();
+	for (UObject* Instance : AllInstancesOfClasses)
+	{
+		ApplyAllTask.EnterProgressFrame(1.f, FText::Format(LOCTEXT("ApplyToSlow", "Applying transforms to {0}"), FText::FromString(Instance->GetPathName())));
+		if (UObject* Ido = Repository.FindInstanceDataObject(Instance))
+		{
+			for (const TPair<FTopLevelAssetPath, UE::FInstanceDataTransformSet>& TransformSet : *StagedTransforms)
+			{
+				UE::FInstanceDataTransforms::Get().ApplyTransformSet(TransformSet.Value, Ido, Instance);
+			}
+		}
+	}
+
+	return FReply::Handled();
+}
+
+FReply SDataRecoveryTool::OnSaveAndApplyClicked() const
+{
+	// save then apply transforms to all instances of this class
+	// TODO: struct transforms wont be applied to instances of that struct in other classes. This will require a reload unfortunately
+	for (const TPair<FTopLevelAssetPath, UE::FInstanceDataTransformSet>& TransformSet : *StagedTransforms)
+	{
+		UE::FInstanceDataTransforms::Get().SaveTransformSet(TransformSet.Value);
+	}
+
+	using namespace UE::Editor::DataStorage::Queries;
+
+	// for now just apply the transforms to all instances of every listed class. This may be slow!
+	ICoreProvider* Storage = UE::Editor::DataStorage::GetMutableDataStorageFeature<ICoreProvider>(StorageFeatureName);
+	UE::Editor::DataStorage::QueryHandle Query = Storage->RegisterQuery(
+		Select()
+		.ReadOnly<FTypedElementClassTypeInfoColumn>()
+		.Where()
+		.All<FHasUnknownPropertiesTag>()
+		.Compile());
+
+	TArray<UObject*> AllInstancesOfClasses;
+	Storage->RunQuery(
+		Query,
+		CreateDirectQueryCallbackBinding(
+			[&AllInstancesOfClasses](IDirectQueryContext&, const FTypedElementClassTypeInfoColumn& TypeInfo)
+			{
+				if (const UClass* Class = TypeInfo.TypeInfo.Get())
+				{
+					GetObjectsOfClass(Class, AllInstancesOfClasses);
+				}
+			})
+	);
+
+	FScopedSlowTask PatchAllTask{ static_cast<float>(AllInstancesOfClasses.Num()), LOCTEXT("PatchAllSlow", "Patching objects") };
+	PatchAllTask.MakeDialog();
+
+	UE::FPropertyBagRepository& Repository = UE::FPropertyBagRepository::Get();
+	for (UObject* Instance : AllInstancesOfClasses)
+	{
+		PatchAllTask.EnterProgressFrame(1.f, FText::Format(LOCTEXT("PatchObjectSlow", "Patching {0}"), FText::FromString(Instance->GetPathName())));
+		if (UObject* Ido = Repository.FindInstanceDataObject(Instance))
+		{
+			for (const TPair<FTopLevelAssetPath, UE::FInstanceDataTransformSet>& TransformSet : *StagedTransforms)
+			{
+				UE::FInstanceDataTransforms::Get().PatchInstanceDataObject(Ido, Instance);
+			}
+		}
+	};
+	return FReply::Handled();
+}
+
 
 TSharedRef<SWidget> SDataRecoveryTool::MakePanel(const FText& InTitle, const TSharedRef<SWidget>& InContent, const TOptional<FOnGetContent>& InMenuContent)
 {
@@ -272,7 +365,7 @@ void SDataRecoveryTool::OnListSelectionChanged(UE::Editor::DataStorage::RowHandl
 		return;
 	}
 
-	ICoreProvider* Storage = GetMutableDataStorageFeature<ICoreProvider>(StorageFeatureName);
+	ICoreProvider* Storage = UE::Editor::DataStorage::GetMutableDataStorageFeature<ICoreProvider>(StorageFeatureName);
 
 	FTypedElementClassTypeInfoColumn* TypeColumn = Storage->GetColumn<FTypedElementClassTypeInfoColumn>(RowHandle);
 
@@ -289,39 +382,9 @@ void SDataRecoveryTool::OnListSelectionChanged(UE::Editor::DataStorage::RowHandl
 
 	*SelectedClassPath = Type->GetClassPathName();
 
-	if (DiffView == nullptr)
-	{
-		return;
-	}
-
 	DiffView->GenerateDetailsViews();
 }
 
-void SDataRecoveryTool::SetSelectedClass()
-{
-	FTSTicker::GetCoreTicker().AddTicker(
-		FTickerDelegate::CreateLambda([this](float DeltaTime) {
-
-			using namespace UE::Editor::DataStorage;
-			ICoreProvider* Storage = GetMutableDataStorageFeature<ICoreProvider>(StorageFeatureName);
-			
-			const UClass* LoadedClass = Cast<UClass>(StaticLoadObject(UClass::StaticClass(), nullptr, SelectedClassPath->ToString()));
-
-			for (RowHandle RowHandle : ReferenceQueryResultsNode->GetRows())
-			{
-				if (Storage->GetColumn<FTypedElementClassTypeInfoColumn>(RowHandle)->TypeInfo == LoadedClass)
-				{
-					ClassesTableViewer->SetSelection(RowHandle, true, ESelectInfo::Direct);
-					break;
-				}
-			}
-
-			return false;
-		}),
-		0.0f
-	);
-}
-
 TSharedRef<SWidget> SDataRecoveryTool::MakeTEDSTableViewer()
 {
 	using namespace UE::Editor::DataStorage;
@@ -340,15 +403,15 @@ TSharedRef<SWidget> SDataRecoveryTool::MakeTEDSTableViewer()
 	TSharedRef<QueryStack::FQueryNode> ReferenceQueryNode =
 		MakeShared<QueryStack::FQueryNode>(*Storage, Select().Where().All<FHasUnknownPropertiesTag, FTypedElementClassTypeInfoColumn>().Compile());
 
-	ReferenceQueryResultsNode = MakeShared<QueryStack::FRowQueryResultsNode>(*Storage,
-		ReferenceQueryNode,
-		QueryStack::FRowQueryResultsNode::ESyncFlags::RefreshOnUpdate);
+	TSharedRef<QueryStack::FRowQueryResultsNode> ReferenceQueryResultsNode =
+		MakeShared<QueryStack::FRowQueryResultsNode>(*Storage,
+			ReferenceQueryNode, QueryStack::FRowQueryResultsNode::ESyncFlags::RefreshOnUpdate);
 
 	TSharedPtr<QueryStack::IRowNode> SearchNode;
 	TSharedPtr<SBox> SearchBoxContainer = SNew(SBox);
 
 	SearchBoxContainer->SetContent(SNew(UE::Editor::DataStorage::STedsSearchBox)
-		.InSearchableRowNode(ReferenceQueryResultsNode)
+		.InSearchableRowNode(ReferenceQueryResultsNode.ToSharedPtr())
 		.OutSearchNode(&SearchNode));
 
 
@@ -356,24 +419,19 @@ TSharedRef<SWidget> SDataRecoveryTool::MakeTEDSTableViewer()
 	// ReferenceQueryResultsNode so that it won't be null
 	if (!SearchNode)
 	{
-		SearchNode = ReferenceQueryResultsNode;
+		SearchNode = ReferenceQueryResultsNode.ToSharedPtr();
 	}
 
 	// Using the Teds Outliner Purpose for matching customization for now
 	IUiProvider::FPurposeID PurposeId = IUiProvider::FPurposeInfo("SceneOutliner", "RowLabel", NAME_None).GeneratePurposeID();
 
 
-	ClassesTableViewer = SNew(STedsTableViewer)
+	TSharedPtr<STedsTableViewer> TypeListView = SNew(STedsTableViewer)
 		.QueryStack(SearchNode)
-		.Columns({FTypedElementClassTypeInfoColumn::StaticStruct()})
+		.Columns({ FTypedElementClassTypeInfoColumn::StaticStruct() })
 		.ListSelectionMode(ESelectionMode::Single)
 		.OnSelectionChanged(this, &SDataRecoveryTool::OnListSelectionChanged);
 
-	if (SelectedClassPath.IsValid())
-	{
-		SetSelectedClass();
-	}
-
 	return
 		SNew(SVerticalBox)
 		+ SVerticalBox::Slot()
@@ -384,7 +442,7 @@ TSharedRef<SWidget> SDataRecoveryTool::MakeTEDSTableViewer()
 		]
 		+ SVerticalBox::Slot()
 		[
-			ClassesTableViewer.ToSharedRef()
+			TypeListView.ToSharedRef()
 		];
 }
 
@@ -397,11 +455,6 @@ void SDataRecoveryTool::SetDockTab(const TSharedRef<SDockTab>& DockTab)
 	OwningDockTab = DockTab;
 }
 
-FReply SDataRecoveryTool::OnConfirmClicked() const
-{
-	return FReply::Handled();
-}
-
 END_SLATE_FUNCTION_BUILD_OPTIMIZATION
 
 
diff --git a/Engine/Plugins/Experimental/InstanceDataObjectFixupTool/Source/Public/DataRecoveryTool.h b/Engine/Plugins/Experimental/InstanceDataObjectFixupTool/Source/Public/DataRecoveryTool.h
--- a/Engine/Plugins/Experimental/InstanceDataObjectFixupTool/Source/Public/DataRecoveryTool.h
+++ b/Engine/Plugins/Experimental/InstanceDataObjectFixupTool/Source/Public/DataRecoveryTool.h
@@ -8,17 +8,15 @@
 #include "Widgets/Docking/SDockTab.h"
 #include "Widgets/SCompoundWidget.h"
 #include "DataRecoveryToolUtils.h"
-#include "UObject/InstanceDataTransforms.h"
 
 #define UE_API INSTANCEDATAOBJECTFIXUPTOOL_API
 
-namespace UE::Editor::DataStorage::QueryStack
+class SInstanceDataObjectFixupTool;
+namespace UE
 {
-	class FRowQueryResultsNode;
+	class FInstanceDataTransformSet;
 }
 
-class SInstanceDataObjectFixupTool;
-
 class SDataRecoveryTool : public SCompoundWidget
 {
 public:
@@ -43,25 +41,22 @@ private:
 	UE_API TSharedRef<SWidget> MakeClassesPanel();
 	UE_API TSharedRef<SWidget> MakeTEDSTableViewer();
 	UE_API void OnListSelectionChanged(UE::Editor::DataStorage::RowHandle RowItem);
-	void SetSelectedClass();
 
 	UE_API TSharedRef<SWidget> MakeRightPanel();
 
-	UE_API FReply OnSaveClicked();
-
 	UE_API TSharedRef<SWidget> MakeDiffSection();
 
-	UE_API FReply OnConfirmClicked() const;
+	FReply OnApplyClicked() const;
+	FReply OnSaveAndApplyClicked() const;
 
 
 	TWeakPtr<SDockTab> OwningDockTab;
 
 	TSharedPtr<FTopLevelAssetPath> SelectedClassPath;
+	TSharedPtr<TMap<FTopLevelAssetPath, UE::FInstanceDataTransformSet>> StagedTransforms;
 
 	TSharedPtr<UE::Editor::DataStorage::STedsTableViewer> ClassesTableViewer;
 
-	TSharedPtr<UE::Editor::DataStorage::QueryStack::FRowQueryResultsNode> ReferenceQueryResultsNode;
-
 	TSharedPtr<SInstanceDataObjectFixupTool> DiffView;
 };
 
