[Feature] DRT - Snapshot instance in details view and other changes

- Added snapshot IDO in the details panel so changes to it are reflected in the diff view but not on an actual object
- Added reset button to undo all actions on that IDO and transformSet
- Added new way for opening the DRT from the Top Bar > Tools menu
- Currently selected class is selected by default in the Classes Table when opening the DRT in the details panel
- Removing HasUnknownProperties after an Apply or Save so it disappears from the class selection if there is no instances of that type to fixup
- Removed some unused code

diff --git a/Engine/Plugins/Experimental/InstanceDataObjectFixupTool/Source/Private/DataRecoveryTool.cpp b/Engine/Plugins/Experimental/InstanceDataObjectFixupTool/Source/Private/DataRecoveryTool.cpp
--- a/Engine/Plugins/Experimental/InstanceDataObjectFixupTool/Source/Private/DataRecoveryTool.cpp
+++ b/Engine/Plugins/Experimental/InstanceDataObjectFixupTool/Source/Private/DataRecoveryTool.cpp
@@ -17,7 +17,9 @@
 #include "SDetailsSplitter.h"
 #include "SlateOptMacros.h"
 #include "SPositiveActionButton.h"
+#include "TedsAlertColumns.h"
 #include "TedsQueryNode.h"
+#include "Elements/Columns/TypedElementMiscColumns.h"
 #include "TedsQueryStack/Public/TedsRowQueryResultsNode.h"
 #include "TedsTableViewer/Public/Widgets/STedsSearchBox.h"
 #include "TedsTableViewer/Public/Widgets/STedsTableViewer.h"
@@ -28,7 +30,6 @@
 #include "Serialization/ObjectWriter.h"
 #include "Misc/ScopedSlowTask.h"
 
-
 #define LOCTEXT_NAMESPACE "DataRecoveryTool"
 
 BEGIN_SLATE_FUNCTION_BUILD_OPTIMIZATION
@@ -188,49 +189,68 @@ TSharedRef<SWidget> SDataRecoveryTool::MakeDiffSection()
 		.Padding(Padding::None)
 		[
 			SAssignNew(DiffView, SInstanceDataObjectFixupTool)
-				.SelectedClassPath(SelectedClassPath)
-				.StagedTransforms(StagedTransforms)
+			.SelectedClassPath(SelectedClassPath)
+			.StagedTransforms(StagedTransforms)
 		];
 }
 
+
 FReply SDataRecoveryTool::OnApplyClicked() const
 {
+	using namespace UE::Editor::DataStorage;
 	using namespace UE::Editor::DataStorage::Queries;
-	
+
 	// for now just apply the transforms to all instances of every listed class. This may be slow!
-	ICoreProvider* Storage = UE::Editor::DataStorage::GetMutableDataStorageFeature<ICoreProvider>(StorageFeatureName);
-	UE::Editor::DataStorage::QueryHandle Query = Storage->RegisterQuery(
-		Select()
-		.ReadOnly<FTypedElementClassTypeInfoColumn>()
-		.Where()
-		.All<FHasUnknownPropertiesTag>()
-		.Compile());
-
-	TArray<UObject*> AllInstancesOfClasses;
-	Storage->RunQuery(
-		Query,
-		CreateDirectQueryCallbackBinding(
-		[&AllInstancesOfClasses](IDirectQueryContext&, const FTypedElementClassTypeInfoColumn& TypeInfo)
-		{
-			if (const UClass* Class = TypeInfo.TypeInfo.Get())
-			{
-				GetObjectsOfClass(Class, AllInstancesOfClasses);
-			}
-		})
-	);
 
-	FScopedSlowTask ApplyAllTask{ static_cast<float>(AllInstancesOfClasses.Num()), LOCTEXT("ApplyAllSlow","Applying transforms") };
+	ICoreProvider* Storage = GetMutableDataStorageFeature<ICoreProvider>(StorageFeatureName);
+
+	FScopedSlowTask ApplyAllTask{ static_cast<float>(ReferenceQueryResultsNode->GetRows().Num()), LOCTEXT("ApplyAllSlow","Applying transforms") };
 	ApplyAllTask.MakeDialog();
 
-	UE::FPropertyBagRepository& Repository = UE::FPropertyBagRepository::Get();
-	for (UObject* Instance : AllInstancesOfClasses)
+	for (const RowHandle Row : ReferenceQueryResultsNode->GetRows())
 	{
-		ApplyAllTask.EnterProgressFrame(1.f, FText::Format(LOCTEXT("ApplyToSlow", "Applying transforms to {0}"), FText::FromString(Instance->GetPathName())));
-		if (UObject* Ido = Repository.FindInstanceDataObject(Instance))
+		const FTypedElementClassTypeInfoColumn* TypeInfo = Storage->GetColumn<FTypedElementClassTypeInfoColumn>(Row);
+		if (!TypeInfo)
 		{
-			for (const TPair<FTopLevelAssetPath, UE::FInstanceDataTransformSet>& TransformSet : *StagedTransforms)
+			continue;
+		}
+
+		if (const UClass* Class = TypeInfo->TypeInfo.Get())
+		{
+			bool bClassStillRequiresFixup = false;
+			const FTopLevelAssetPath ClassPath = Class->GetClassPathName();
+
+			ApplyAllTask.EnterProgressFrame(1.f, FText::Format(LOCTEXT("ApplyToSlow", "Applying transforms to {0}"), FText::FromString(ClassPath.ToString())));
+
+			const UE::FInstanceDataTransformSet* TransformSet = StagedTransforms->Find(ClassPath);
+
+			if (TransformSet == nullptr)
+			{
+				continue;
+			}
+
+			TArray<UObject*> AllInstancesOfClasses;
+
+			// Could not use ApplyTransformSet inside ForEachObjectOfClass lambda, so we have to get the objects first
+			GetObjectsOfClass(Class, AllInstancesOfClasses);
+
+			for (UObject* Instance : AllInstancesOfClasses)
+			{
+				if (UObject* Ido = UE::FPropertyBagRepository::Get().FindInstanceDataObject(Instance))
+				{
+					UE::FInstanceDataTransforms::Get().ApplyTransformSet(*TransformSet, Ido, Instance);
+
+					if (!bClassStillRequiresFixup && UE::FPropertyBagRepository::Get().RequiresFixup(Instance))
+					{
+						bClassStillRequiresFixup = true;
+					}
+				}
+			}
+
+			if (!bClassStillRequiresFixup)
 			{
-				UE::FInstanceDataTransforms::Get().ApplyTransformSet(TransformSet.Value, Ido, Instance);
+				Storage->RemoveColumn<FHasUnknownPropertiesTag>(Row);
+				Storage->AddColumn<FTypedElementSyncBackToWorldTag>(Row);
 			}
 		}
 	}
@@ -247,48 +267,66 @@ FReply SDataRecoveryTool::OnSaveAndApplyClicked() const
 		UE::FInstanceDataTransforms::Get().SaveTransformSet(TransformSet.Value);
 	}
 
+	using namespace UE::Editor::DataStorage;
 	using namespace UE::Editor::DataStorage::Queries;
 
 	// for now just apply the transforms to all instances of every listed class. This may be slow!
-	ICoreProvider* Storage = UE::Editor::DataStorage::GetMutableDataStorageFeature<ICoreProvider>(StorageFeatureName);
-	UE::Editor::DataStorage::QueryHandle Query = Storage->RegisterQuery(
-		Select()
-		.ReadOnly<FTypedElementClassTypeInfoColumn>()
-		.Where()
-		.All<FHasUnknownPropertiesTag>()
-		.Compile());
-
-	TArray<UObject*> AllInstancesOfClasses;
-	Storage->RunQuery(
-		Query,
-		CreateDirectQueryCallbackBinding(
-			[&AllInstancesOfClasses](IDirectQueryContext&, const FTypedElementClassTypeInfoColumn& TypeInfo)
-			{
-				if (const UClass* Class = TypeInfo.TypeInfo.Get())
-				{
-					GetObjectsOfClass(Class, AllInstancesOfClasses);
-				}
-			})
-	);
 
-	FScopedSlowTask PatchAllTask{ static_cast<float>(AllInstancesOfClasses.Num()), LOCTEXT("PatchAllSlow", "Patching objects") };
+	ICoreProvider* Storage = GetMutableDataStorageFeature<ICoreProvider>(StorageFeatureName);
+
+	FScopedSlowTask PatchAllTask{ static_cast<float>(ReferenceQueryResultsNode->GetRows().Num()), LOCTEXT("PatchAllSlow","Patching transforms") };
 	PatchAllTask.MakeDialog();
 
-	UE::FPropertyBagRepository& Repository = UE::FPropertyBagRepository::Get();
-	for (UObject* Instance : AllInstancesOfClasses)
+	for (const RowHandle Row : ReferenceQueryResultsNode->GetRows())
 	{
-		PatchAllTask.EnterProgressFrame(1.f, FText::Format(LOCTEXT("PatchObjectSlow", "Patching {0}"), FText::FromString(Instance->GetPathName())));
-		if (UObject* Ido = Repository.FindInstanceDataObject(Instance))
+		const FTypedElementClassTypeInfoColumn* TypeInfo = Storage->GetColumn<FTypedElementClassTypeInfoColumn>(Row);
+		if (!TypeInfo)
+		{
+			continue;
+		}
+
+		if (const UClass* Class = TypeInfo->TypeInfo.Get())
 		{
-			for (const TPair<FTopLevelAssetPath, UE::FInstanceDataTransformSet>& TransformSet : *StagedTransforms)
+			bool bClassStillRequiresFixup = false;
+			const FTopLevelAssetPath ClassPath = Class->GetClassPathName();
+
+			PatchAllTask.EnterProgressFrame(1.f, FText::Format(LOCTEXT("PatchToSlow", "Patching transforms to {0}"), FText::FromString(ClassPath.ToString())));
+
+			const UE::FInstanceDataTransformSet* TransformSet = StagedTransforms->Find(ClassPath);
+
+			if (TransformSet == nullptr)
 			{
-				UE::FInstanceDataTransforms::Get().PatchInstanceDataObject(Ido, Instance);
+				continue;
+			}
+
+			TArray<UObject*> AllInstancesOfClasses;
+
+			// Could not use PatchTransformSet inside ForEachObjectOfClass lambda, so we have to get the objects first
+			GetObjectsOfClass(Class, AllInstancesOfClasses);
+
+			for (UObject* Instance : AllInstancesOfClasses)
+			{
+				if (UObject* Ido = UE::FPropertyBagRepository::Get().FindInstanceDataObject(Instance))
+				{
+					UE::FInstanceDataTransforms::Get().PatchInstanceDataObject(Ido, Instance);
+
+					if (!bClassStillRequiresFixup && UE::FPropertyBagRepository::Get().RequiresFixup(Instance))
+					{
+						bClassStillRequiresFixup = true;
+					}
+				}
+			}
+
+			if (!bClassStillRequiresFixup)
+			{
+				Storage->RemoveColumn<FHasUnknownPropertiesTag>(Row);
+				Storage->AddColumn<FTypedElementSyncBackToWorldTag>(Row);
 			}
 		}
-	};
-	return FReply::Handled();
-}
+	}
 
+	return OnApplyClicked();
+}
 
 TSharedRef<SWidget> SDataRecoveryTool::MakePanel(const FText& InTitle, const TSharedRef<SWidget>& InContent, const TOptional<FOnGetContent>& InMenuContent)
 {
@@ -360,12 +398,15 @@ void SDataRecoveryTool::OnListSelectionChanged(UE::Editor::DataStorage::RowHandl
 	{
 		*SelectedClassPath = nullptr;
 
-		DiffView->GenerateDetailsViews();
+		if (DiffView != nullptr)
+		{
+			DiffView->GenerateDetailsViews();
+		}
 
 		return;
 	}
 
-	ICoreProvider* Storage = UE::Editor::DataStorage::GetMutableDataStorageFeature<ICoreProvider>(StorageFeatureName);
+	ICoreProvider* Storage = GetMutableDataStorageFeature<ICoreProvider>(StorageFeatureName);
 
 	FTypedElementClassTypeInfoColumn* TypeColumn = Storage->GetColumn<FTypedElementClassTypeInfoColumn>(RowHandle);
 
@@ -382,9 +423,41 @@ void SDataRecoveryTool::OnListSelectionChanged(UE::Editor::DataStorage::RowHandl
 
 	*SelectedClassPath = Type->GetClassPathName();
 
+	if (DiffView == nullptr)
+	{
+		return;
+	}
+
 	DiffView->GenerateDetailsViews();
 }
 
+void SDataRecoveryTool::SetSelectedClass()
+{
+	if (!ClassesTableViewer)
+	{
+		return;
+	}
+	FTSTicker::GetCoreTicker().AddTicker(
+		FTickerDelegate::CreateLambda([this](float DeltaTime) {
+			using namespace UE::Editor::DataStorage;
+			ICoreProvider* Storage = GetMutableDataStorageFeature<ICoreProvider>(StorageFeatureName);
+
+			const UClass* LoadedClass = Cast<UClass>(StaticLoadObject(UClass::StaticClass(), nullptr, SelectedClassPath->ToString()));
+			for (RowHandle RowHandle : ReferenceQueryResultsNode->GetRows())
+			{
+				FTypedElementClassTypeInfoColumn* TypeInfo = Storage->GetColumn<FTypedElementClassTypeInfoColumn>(RowHandle);
+				if (TypeInfo && TypeInfo->TypeInfo == LoadedClass)
+				{
+					ClassesTableViewer->SetSelection(RowHandle, true, ESelectInfo::Direct);
+					break;
+				}
+			}
+			return false;
+			}),
+		0.0f
+	);
+}
+
 TSharedRef<SWidget> SDataRecoveryTool::MakeTEDSTableViewer()
 {
 	using namespace UE::Editor::DataStorage;
@@ -403,15 +476,15 @@ TSharedRef<SWidget> SDataRecoveryTool::MakeTEDSTableViewer()
 	TSharedRef<QueryStack::FQueryNode> ReferenceQueryNode =
 		MakeShared<QueryStack::FQueryNode>(*Storage, Select().Where().All<FHasUnknownPropertiesTag, FTypedElementClassTypeInfoColumn>().Compile());
 
-	TSharedRef<QueryStack::FRowQueryResultsNode> ReferenceQueryResultsNode =
-		MakeShared<QueryStack::FRowQueryResultsNode>(*Storage,
-			ReferenceQueryNode, QueryStack::FRowQueryResultsNode::ESyncFlags::RefreshOnUpdate);
+	ReferenceQueryResultsNode = MakeShared<QueryStack::FRowQueryResultsNode>(*Storage,
+		ReferenceQueryNode,
+		QueryStack::FRowQueryResultsNode::ESyncFlags::RefreshOnUpdate);
 
 	TSharedPtr<QueryStack::IRowNode> SearchNode;
 	TSharedPtr<SBox> SearchBoxContainer = SNew(SBox);
 
 	SearchBoxContainer->SetContent(SNew(UE::Editor::DataStorage::STedsSearchBox)
-		.InSearchableRowNode(ReferenceQueryResultsNode.ToSharedPtr())
+		.InSearchableRowNode(ReferenceQueryResultsNode)
 		.OutSearchNode(&SearchNode));
 
 
@@ -419,19 +492,26 @@ TSharedRef<SWidget> SDataRecoveryTool::MakeTEDSTableViewer()
 	// ReferenceQueryResultsNode so that it won't be null
 	if (!SearchNode)
 	{
-		SearchNode = ReferenceQueryResultsNode.ToSharedPtr();
+		SearchNode = ReferenceQueryResultsNode;
 	}
 
 	// Using the Teds Outliner Purpose for matching customization for now
 	IUiProvider::FPurposeID PurposeId = IUiProvider::FPurposeInfo("SceneOutliner", "RowLabel", NAME_None).GeneratePurposeID();
 
-
-	TSharedPtr<STedsTableViewer> TypeListView = SNew(STedsTableViewer)
+	ClassesTableViewer = SNew(STedsTableViewer)
 		.QueryStack(SearchNode)
-		.Columns({ FTypedElementClassTypeInfoColumn::StaticStruct() })
+		.EmptyRowsMessage(LOCTEXT("EmptyRowsMessage", "No loose properties found"))
+		.Columns({ FTypedElementClassTypeInfoColumn::StaticStruct(), FTedsAlertColumn::StaticStruct() })
 		.ListSelectionMode(ESelectionMode::Single)
 		.OnSelectionChanged(this, &SDataRecoveryTool::OnListSelectionChanged);
 
+
+
+	if (SelectedClassPath.IsValid())
+	{
+		SetSelectedClass();
+	}
+
 	return
 		SNew(SVerticalBox)
 		+ SVerticalBox::Slot()
@@ -442,7 +522,7 @@ TSharedRef<SWidget> SDataRecoveryTool::MakeTEDSTableViewer()
 		]
 		+ SVerticalBox::Slot()
 		[
-			TypeListView.ToSharedRef()
+			ClassesTableViewer.ToSharedRef()
 		];
 }
 
diff --git a/Engine/Plugins/Experimental/InstanceDataObjectFixupTool/Source/Private/DataRecoveryToolUtils.h b/Engine/Plugins/Experimental/InstanceDataObjectFixupTool/Source/Private/DataRecoveryToolUtils.h
--- a/Engine/Plugins/Experimental/InstanceDataObjectFixupTool/Source/Private/DataRecoveryToolUtils.h
+++ b/Engine/Plugins/Experimental/InstanceDataObjectFixupTool/Source/Private/DataRecoveryToolUtils.h
@@ -10,12 +10,19 @@
 
 namespace UE::DataRecoveryTool::Utils
 {
-	inline TObjectPtr<UObject> GetInstanceDataObjectFromSelectedClassPath(const FTopLevelAssetPath& ClassPath)
+	inline TObjectPtr<UObject> GetInstanceDataObjectFromSelectedClassPath(const FTopLevelAssetPath& ClassPath, TObjectPtr<UObject>* OutOwner = nullptr)
 	{
 		const UClass* LoadedClass = Cast<UClass>(StaticLoadObject(UClass::StaticClass(), nullptr, ClassPath.ToString()));
 
 		const UObject* FoundInstance = FPropertyBagRepository::Get().FindFirstInstanceThatNeedsFixup(LoadedClass);
 
+		if (OutOwner)
+		{
+			*OutOwner = const_cast<UObject*>(FoundInstance);
+		}
+
 		return FoundInstance ? FPropertyBagRepository::Get().FindInstanceDataObject(FoundInstance) : nullptr;
 	}
+
+	
 }
diff --git a/Engine/Plugins/Experimental/InstanceDataObjectFixupTool/Source/Private/InstanceDataObjectFixupDetailCustomization.cpp b/Engine/Plugins/Experimental/InstanceDataObjectFixupTool/Source/Private/InstanceDataObjectFixupDetailCustomization.cpp
--- a/Engine/Plugins/Experimental/InstanceDataObjectFixupTool/Source/Private/InstanceDataObjectFixupDetailCustomization.cpp
+++ b/Engine/Plugins/Experimental/InstanceDataObjectFixupTool/Source/Private/InstanceDataObjectFixupDetailCustomization.cpp
@@ -33,26 +33,12 @@ TSharedRef<SWidget> FInstanceDataObjectNameWidgetOverride::CustomizeName(TShared
 		]
 		+SWidgetSwitcher::Slot()
 		[
-			SNew(SHorizontalBox)
-			+SHorizontalBox::Slot()
-			.VAlign(EVerticalAlignment::VAlign_Center)
-			.AutoWidth()
+			SNew(SComboButton)
+			.Visibility(EVisibility::Visible)
+			.OnGetMenuContent_Raw(this, &FInstanceDataObjectNameWidgetOverride::GeneratePropertyRedirectMenu, Path)
+			.ButtonContent()
 			[
-				SNew(STextBlock)
-				.Visibility(this, &FInstanceDataObjectNameWidgetOverride::DeletionSymbolVisibility, Path)
-				.Text(LOCTEXT("MarkedForDeletion", "❌"))
-			]
-			+SHorizontalBox::Slot()
-			.VAlign(EVerticalAlignment::VAlign_Center)
-			.AutoWidth()
-			[
-				SNew(SComboButton)
-				.Visibility(EVisibility::Visible)
-				.OnGetMenuContent_Raw(this, &FInstanceDataObjectNameWidgetOverride::GeneratePropertyRedirectMenu, Path)
-				.ButtonContent()
-				[
-					InnerNameContent
-				]
+				InnerNameContent
 			]
 		];
 	return NameContent;
@@ -166,10 +152,6 @@ int32 FInstanceDataObjectNameWidgetOverride::GetNameWidgetIndex(FPropertyPath Pa
 	{
 		if (Panel->HasViewFlag(FInstanceDataObjectFixupPanel::EViewFlags::AllowRemapLooseProperties))
 		{
-			if (Panel->RevertInfo.Contains(Path))
-			{
-				return DisplayRedirectMenu;
-			}
 			if (const FProperty* Property = Path.GetLeafMostProperty().Property.Get())
 			{
 				if (Property->GetBoolMetaData(TEXT("isLoose")))
@@ -217,31 +199,14 @@ TSharedRef<SWidget> FInstanceDataObjectNameWidgetOverride::GeneratePropertyRedir
 	{
 		return MenuBuilder.MakeWidget();
 	}
-	const FPropertyPath& OriginalPath = Panel->GetOriginalPath(Path);
-
-	MenuBuilder.BeginSection(NAME_None, LOCTEXT("ResetRedirect", "Reset"));
-	if (OriginalPath != Path || Panel->MarkedForDelete.Contains(OriginalPath))
-	{
-		FText OriginalPathText = GetCleanVersePath(OriginalPath);
-		FText Tooltip = FText::Format(LOCTEXT("ResetTooltip", "Reset back to {0}"), OriginalPathText);
-		MenuBuilder.AddMenuEntry(OriginalPathText, Tooltip, FSlateIcon()
-						, FUIAction(FExecuteAction::CreateSP(Panel.Get(), &FInstanceDataObjectFixupPanel::OnRedirectProperty, Path, OriginalPath))
-						, NAME_None
-						, EUserInterfaceActionType::RadioButton);
-	}
-	MenuBuilder.EndSection();
 	
 	MenuBuilder.BeginSection(NAME_None, LOCTEXT("Delete", "Delete"));
 	{
-		if (!Panel->MarkedForDelete.Contains(OriginalPath)) // check that it wasn't already marked as deleted
-		{
-			FText DisplayName = LOCTEXT("MarkForDeletion", "Mark For Deletion");
-			FText Tooltip = LOCTEXT("MarkForDeletionTooltip", "Mark this property for deletion");
-			MenuBuilder.AddMenuEntry(DisplayName, Tooltip, FSlateIcon()
-						, FUIAction(FExecuteAction::CreateSP(Panel.Get(), &FInstanceDataObjectFixupPanel::OnMarkForDelete, Path))
-						, NAME_None
-						, EUserInterfaceActionType::RadioButton);
-		}
+		FText DisplayName = LOCTEXT("MarkForDeletion", "Mark For Deletion");
+		FText Tooltip = LOCTEXT("MarkForDeletionTooltip", "Mark this property for deletion");
+		MenuBuilder.AddMenuEntry(DisplayName, Tooltip, FSlateIcon(FAppStyle::GetAppStyleSetName(), FName("Icons.Delete"))
+					, FUIAction(FExecuteAction::CreateSP(Panel.Get(), &FInstanceDataObjectFixupPanel::OnMarkForDelete, Path))
+					, NAME_None);
 	}
 	MenuBuilder.EndSection();
 
@@ -268,10 +233,9 @@ TSharedRef<SWidget> FInstanceDataObjectNameWidgetOverride::GeneratePropertyRedir
 				{
 					FText DisplayName = GetCleanVersePath(Option);
 					FText Tooltip = FText::Format(LOCTEXT("MovePropertyTooltip", "Move property to '{0}'"), DisplayName);
-					MenuBuilder.AddMenuEntry(DisplayName, Tooltip, FSlateIcon()
+					MenuBuilder.AddMenuEntry(DisplayName, Tooltip, FSlateIcon(FAppStyle::GetAppStyleSetName(), FName("Icons.Convert"))
 					, FUIAction(FExecuteAction::CreateSP(Panel.Get(), &FInstanceDataObjectFixupPanel::OnRedirectProperty, Path, Option))
-					, NAME_None
-					, EUserInterfaceActionType::RadioButton);
+					, NAME_None);
 				}
 			}
 		}
@@ -296,10 +260,9 @@ TSharedRef<SWidget> FInstanceDataObjectNameWidgetOverride::GeneratePropertyRedir
 			{
 				FText DisplayName = GetCleanVersePath(Option);
 				FText Tooltip = FText::Format(LOCTEXT("RenamePropertyTooltip", "Rename property to '{0}'"), DisplayName);
-				MenuBuilder.AddMenuEntry(DisplayName, Tooltip, FSlateIcon()
+				MenuBuilder.AddMenuEntry(DisplayName, Tooltip, FSlateIcon(FAppStyle::GetAppStyleSetName(), FName("Icons.Convert"))
 				, FUIAction(FExecuteAction::CreateSP(Panel.Get(), &FInstanceDataObjectFixupPanel::OnRedirectProperty, Path, Option))
-				, NAME_None
-				, EUserInterfaceActionType::RadioButton);
+				, NAME_None);
 			}
 		}
 	}
@@ -372,27 +335,9 @@ TSharedRef<SWidget> FInstanceDataObjectNameWidgetOverride::GeneratePropertyRedir
 	//		}
 	//	}
 	//}
-	MenuBuilder.EndSection();
+	//MenuBuilder.EndSection();
 	
 	return MenuBuilder.MakeWidget();
 }
 
-EVisibility FInstanceDataObjectNameWidgetOverride::DeletionSymbolVisibility(FPropertyPath Path) const
-{
-	if (const TSharedPtr<FInstanceDataObjectFixupPanel> Panel = DiffPanel.Pin())
-	{
-		return Panel->MarkedForDelete.Contains(Path) ? EVisibility::Visible :  EVisibility::Collapsed;
-	}
-	return EVisibility::Collapsed;
-}
-
-EVisibility FInstanceDataObjectNameWidgetOverride::ValueContentVisibility(FPropertyPath Path) const
-{
-	if (const TSharedPtr<FInstanceDataObjectFixupPanel> Panel = DiffPanel.Pin())
-	{
-		return Panel->MarkedForDelete.Contains(Path) ? EVisibility::Collapsed :  EVisibility::Visible;
-	}
-	return EVisibility::Visible;
-}
-
 #undef LOCTEXT_NAMESPACE
diff --git a/Engine/Plugins/Experimental/InstanceDataObjectFixupTool/Source/Private/InstanceDataObjectFixupDetailCustomization.h b/Engine/Plugins/Experimental/InstanceDataObjectFixupTool/Source/Private/InstanceDataObjectFixupDetailCustomization.h
--- a/Engine/Plugins/Experimental/InstanceDataObjectFixupTool/Source/Private/InstanceDataObjectFixupDetailCustomization.h
+++ b/Engine/Plugins/Experimental/InstanceDataObjectFixupTool/Source/Private/InstanceDataObjectFixupDetailCustomization.h
@@ -26,8 +26,6 @@ private:
 	
 	UE_API int32 GetNameWidgetIndex(FPropertyPath Path) const;
 	UE_API TSharedRef<SWidget> GeneratePropertyRedirectMenu(FPropertyPath Path) const;
-	UE_API EVisibility DeletionSymbolVisibility(FPropertyPath Path) const;
-	UE_API EVisibility ValueContentVisibility(FPropertyPath Path) const;
 	
 	UE_API TSet<FPropertyPath> GetRedirectOptions(const UStruct* Struct, void* Value) const;
 	UE_API void GetRedirectOptions(const UStruct* Struct, void* Value, const FPropertyPath& Path, TSet<FPropertyPath>& OutPaths) const;
diff --git a/Engine/Plugins/Experimental/InstanceDataObjectFixupTool/Source/Private/InstanceDataObjectFixupPanel.cpp b/Engine/Plugins/Experimental/InstanceDataObjectFixupTool/Source/Private/InstanceDataObjectFixupPanel.cpp
--- a/Engine/Plugins/Experimental/InstanceDataObjectFixupTool/Source/Private/InstanceDataObjectFixupPanel.cpp
+++ b/Engine/Plugins/Experimental/InstanceDataObjectFixupTool/Source/Private/InstanceDataObjectFixupPanel.cpp
@@ -9,7 +9,6 @@
 #include "InstanceDataObjectFixupDetailCustomization.h"
 #include "Misc/StringOutputDevice.h"
 #include "Modules/ModuleManager.h"
-#include "Editor.h"
 #include "Editor/PropertyEditor/Private/PropertyNode.h"
 #include "TedsAlerts.h"
 #include "UObject/PropertyBagRepository.h"
@@ -19,13 +18,11 @@
 #include "Elements/Interfaces/TypedElementDataStorageInterface.h"
 #include "Elements/Interfaces/TypedElementDataStorageCompatibilityInterface.h"
 #include "Serialization/ObjectReader.h"
-#include "Serialization/ObjectWriter.h"
 #include "Dom/JsonValue.h"
 #include "Dom/JsonObject.h"
 
 #include "UObject/OverriddenPropertySet.h"
 #include "UObject/OverridableManager.h"
-#include "UObject/TextProperty.h"
 #include "UObject/UObjectThreadContext.h"
 
 #include "UObject/PropertyPathNameTree.h"
@@ -34,8 +31,6 @@
 #define LOCTEXT_NAMESPACE "InstanceDataObjectFixupPanel"
 
 static const FName NAME_IsLooseMetadata(TEXT("IsLoose"));
-static const FName NAME_ContainsLoosePropertiesMetadata(ANSITEXTVIEW("ContainsLooseProperties"));
-
 
 static bool IsPathOverridden(UObject* Object, const FPropertyPath& Path)
 {
@@ -69,24 +64,18 @@ static bool IsPathOverridden(UObject* Object, const FPropertyPath& Path)
 	EOverriddenPropertyOperation Override = FOverridableManager::Get().GetOverriddenPropertyOperation(Object, VisitorPath);
 	return (Override != EOverriddenPropertyOperation::None);
 }
-
-
-FInstanceDataObjectFixupPanel::FInstanceDataObjectFixupPanel(const TObjectPtr<UObject>& InstanceDataObject, EViewFlags ViewFlags, TWeakPtr<TMap<FTopLevelAssetPath, UE::FInstanceDataTransformSet>> StagedTransforms)
+FInstanceDataObjectFixupPanel::FInstanceDataObjectFixupPanel(
+	const TObjectPtr<UObject>& InstanceDataObject,
+	EViewFlags ViewFlags, 
+	const TWeakPtr<TMap<FTopLevelAssetPath,
+	UE::FInstanceDataTransformSet>>& StagedTransforms
+)
 	: InstanceDataObject(InstanceDataObject)
 	, ViewFlags(ViewFlags)
 	, StagedTransforms(StagedTransforms)
 {
 }
 
-FInstanceDataObjectFixupPanel::~FInstanceDataObjectFixupPanel()
-{
-	using namespace UE::Editor::DataStorage;
-
-	ICoreProvider* DataStorage = GetMutableDataStorageFeature<ICoreProvider>(StorageFeatureName);
-	ICompatibilityProvider* DataStorageCompatibility = GetMutableDataStorageFeature<ICompatibilityProvider>(CompatibilityFeatureName);
-	UE::FPropertyBagRepository& Repository = UE::FPropertyBagRepository::Get();
-}
-
 static bool RemoveCustomizationsWithLooseProperties(const FFieldVariant& FieldVariant, const TSharedPtr<IDetailsView>& DetailsView)
 {
 #if WITH_EDITORONLY_DATA
@@ -161,6 +150,11 @@ static bool RemoveCustomizationsWithLooseProperties(const FFieldVariant& FieldVa
 	return false;
 }
 
+void FInstanceDataObjectFixupPanel::RefreshDetailsView()
+{
+	DetailsView->SetObject(InstanceDataObject, true);
+}
+
 TSharedPtr<IDetailsView>& FInstanceDataObjectFixupPanel::GenerateDetailsView(bool bScrollbarOnLeft)
 {
 	FDetailsViewArgs DetailsViewArgs;
@@ -171,6 +165,12 @@ TSharedPtr<IDetailsView>& FInstanceDataObjectFixupPanel::GenerateDetailsView(boo
 	DetailsViewArgs.DetailsNameWidgetOverrideCustomization = MakeShared<FInstanceDataObjectNameWidgetOverride>(SharedThis(this));
 	DetailsViewArgs.bResolveInstanceDataObjects.Emplace(true);
 	DetailsViewArgs.bShowLooseProperties = !HasViewFlag(EViewFlags::HideLooseProperties);
+	DetailsViewArgs.bLockable = false;                                               
+	DetailsViewArgs.bAllowSearch = false;                                            
+	DetailsViewArgs.bSearchInitialKeyFocus = false;                                  
+	DetailsViewArgs.bShowOptions = false;                                            
+	DetailsViewArgs.bAllowFavoriteSystem = false;                                    
+	DetailsViewArgs.bAllowMultipleTopLevelObjects = false;
 
 	if (HasViewFlag(EViewFlags::IncludeOnlySetBySerialization))
 	{
@@ -184,9 +184,8 @@ TSharedPtr<IDetailsView>& FInstanceDataObjectFixupPanel::GenerateDetailsView(boo
 	DetailsView = PropertyEditorModule.CreateDetailView(DetailsViewArgs);
 
 	RemoveCustomizationsWithLooseProperties(InstanceDataObject->GetClass(), DetailsView);
-	
-	
-	DetailsView->SetObject(InstanceDataObject, true);
+
+	RefreshDetailsView();
 
 	FCoreUObjectDelegates::OnObjectsReinstanced.AddSPLambda(this,
 		[this](const TMap<UObject*, UObject*>& ObjectMap)
@@ -220,26 +219,12 @@ TSharedPtr<FAsyncDetailViewDiff> FInstanceDataObjectFixupPanel::GetDiffAgainstRi
 	return DiffAgainstRight.Pin();
 }
 
-bool FInstanceDataObjectFixupPanel::ShouldSplitterIgnoreRow(const TWeakPtr<FDetailTreeNode>& WeakDetailTreeNode) const
-{
-	if (const TSharedPtr<FDetailTreeNode> DetailTreeNode = WeakDetailTreeNode.Pin())
-	{
-		if (const TSharedPtr<IPropertyHandle> Handle = DetailTreeNode->CreatePropertyHandle())
-		{
-			if (MarkedForDelete.Contains(*Handle->CreateFPropertyPath()))
-			{
-				return true;
-			}
-		}
-	}
-	return false;
-}
-
 bool FInstanceDataObjectFixupPanel::AreAllConflictsRedirected() const
 {
 	UE::FPropertyBagRepository& Repository = UE::FPropertyBagRepository::Get();
 
-	return !Repository.RequiresFixup(Repository.FindInstanceForDataObject(InstanceDataObject));
+	const UObject* Owner = Repository.FindInstanceForDataObject(InstanceDataObject);
+	return Owner == nullptr || !Repository.RequiresFixup(Owner);
 }
 
 void FInstanceDataObjectFixupPanel::AutoApplyMarkDeletedActions()
@@ -276,76 +261,6 @@ bool FInstanceDataObjectFixupPanel::HasViewFlag(EViewFlags Flag)
 	return static_cast<uint8>(Flag) & static_cast<uint8>(ViewFlags);
 }
 
-static void* ResolvePath(const FPropertyPath& Path, void* Value)
-{
-	const int32 LastPathIndex = Path.GetNumProperties() - 1;
-	for(int32 PathIndex = 0; PathIndex < Path.GetNumProperties(); ++PathIndex)
-	{
-		const FPropertyInfo& PropertyInfo = Path.GetPropertyInfo(PathIndex);
-		const FProperty* Property = PropertyInfo.Property.Get();
-		if (!Property)
-		{
-			return nullptr;
-		}
-
-		Value = Property->ContainerPtrToValuePtr<void>(Value, PropertyInfo.ArrayIndex != INDEX_NONE ? PropertyInfo.ArrayIndex : 0);
-
-		if (PathIndex < LastPathIndex)
-		{
-			if (const FObjectProperty* AsObjectProperty = CastField<FObjectProperty>(Property))
-			{
-				UObject* Object = AsObjectProperty->GetObjectPropertyValue(Value);
-				UE::FPropertyBagRepository& PropertyBagRepository = UE::FPropertyBagRepository::Get();
-				if (UObject* Found = PropertyBagRepository.FindInstanceDataObject(Object))
-				{
-					Object = Found;
-				}
-				Value = Object;
-			}
-			if (const FArrayProperty* AsArrayProperty = CastField<FArrayProperty>(Property))
-			{
-				FScriptArrayHelper Helper(AsArrayProperty, Value);
-				Value = Helper.GetElementPtr(Path.GetPropertyInfo(++PathIndex).ArrayIndex);
-			}
-			if (const FSetProperty* AsSetProperty = CastField<FSetProperty>(Property))
-			{
-				FScriptSetHelper Helper(AsSetProperty, Value);
-				Value = Helper.FindNthElementPtr(Path.GetPropertyInfo(++PathIndex).ArrayIndex);
-			}
-			if (const FMapProperty* AsMapProperty = CastField<FMapProperty>(Property))
-			{
-				FScriptMapHelper Helper(AsMapProperty, Value);
-				Value = Helper.FindNthValuePtr(Path.GetPropertyInfo(++PathIndex).ArrayIndex);
-			}
-		}
-	}
-	
-	return Value;
-}
-
-static FPropertyChangedEvent ConstructChangeEventForRedirect(const FPropertyPath& Path, FEditPropertyChain& OutChain, TMap<FString, int32>& OutArrayIndices)
-{
-	FPropertyChangedEvent OutEvent = FPropertyChangedEvent(Path.GetLeafMostProperty().Property.Get(), EPropertyChangeType::ValueSet);
-	for (int32 I = 0; I < Path.GetNumProperties(); ++I)
-	{
-		const FPropertyInfo& Info = Path.GetPropertyInfo(I);
-		OutChain.AddTail(Info.Property.Get()); // only the head is used in OverrideProperty
-		if (Info.ArrayIndex != INDEX_NONE)
-		{
-			OutArrayIndices.Add(Info.Property->GetName(), Info.ArrayIndex);
-		}
-		if (Info.Property->IsA<FArrayProperty>() || Info.Property->IsA<FSetProperty>() || Info.Property->IsA<FMapProperty>())
-		{
-			if (++I < Path.GetNumProperties())
-			{
-				OutArrayIndices.Add(Info.Property->GetName(), Path.GetPropertyInfo(I).ArrayIndex);
-			}
-		}
-	}
-	OutEvent.SetArrayIndexPerObject(MakeArrayView(&OutArrayIndices, 1));
-	return OutEvent;
-}
-
 static UStruct* ExtractStructFromType(UE::FPropertyTypeName Type, const UE::FPropertyPathName& Path, int32 SegmentIndex)
 {
 	static const FName NAME_Key(ANSITEXTVIEW("Key"));
@@ -466,7 +381,9 @@ void FInstanceDataObjectFixupPanel::RedirectProperty(const FPropertyPath& From,
 		Transforms->Operations.Add({ .OpCode = TEXT("RedirectProperty"), .Params = Params });
 
 		TObjectPtr<UObject> OwnerObjectPtr = const_cast<UObject*>(Owner);
-		UE::FInstanceDataTransforms::Get().ApplyTransformSet(*Transforms, InstanceDataObject, OwnerObjectPtr);
+		InstanceDataObject = UE::FInstanceDataTransforms::Get().ApplyTransformSet(*Transforms, InstanceDataObject, OwnerObjectPtr);
+
+		RefreshDetailsView();
 	}
 }
 
@@ -499,15 +416,6 @@ bool FInstanceDataObjectFixupPanel::PropertyNeedsRedirect(const FPropertyPath& P
 	return UE::FPropertyBagRepository::Get().RequiresFixup(Owner, PathNames[0]);
 }
 
-const FPropertyPath& FInstanceDataObjectFixupPanel::GetOriginalPath(const FPropertyPath& Path) const
-{
-	if (const FRevertInfo* Found = RevertInfo.Find(Path))
-	{
-		return Found->OriginalPath;
-	}
-	return Path;
-}
-
 void FInstanceDataObjectFixupPanel::MarkForDelete(const FPropertyPath& CurrentPath)
 {
 	FUObjectSerializeContext* SerializeContext = FUObjectThreadContext::Get().GetSerializeContext();
@@ -528,8 +436,10 @@ void FInstanceDataObjectFixupPanel::MarkForDelete(const FPropertyPath& CurrentPa
 		Transforms->Operations.Add({ .OpCode = TEXT("RemoveProperty"), .Params = Params });
 
 		TObjectPtr<UObject> OwnerObjectPtr = const_cast<UObject*>(Owner);
-		UE::FInstanceDataTransforms::Get().ApplyTransformSet(*Transforms, InstanceDataObject, OwnerObjectPtr);
+		InstanceDataObject = UE::FInstanceDataTransforms::Get().ApplyTransformSet(*Transforms, InstanceDataObject, OwnerObjectPtr);
 	}
+
+	RefreshDetailsView();
 }
 
 void FInstanceDataObjectFixupPanel::OnMarkForDelete(FPropertyPath Path)
diff --git a/Engine/Plugins/Experimental/InstanceDataObjectFixupTool/Source/Private/InstanceDataObjectFixupPanel.h b/Engine/Plugins/Experimental/InstanceDataObjectFixupTool/Source/Private/InstanceDataObjectFixupPanel.h
--- a/Engine/Plugins/Experimental/InstanceDataObjectFixupTool/Source/Private/InstanceDataObjectFixupPanel.h
+++ b/Engine/Plugins/Experimental/InstanceDataObjectFixupTool/Source/Private/InstanceDataObjectFixupPanel.h
@@ -27,9 +27,9 @@ public:
 		DefaultRightPanel = HideLooseProperties,
 	};
 	
-	FInstanceDataObjectFixupPanel(const TObjectPtr<UObject>& InstanceDataObject, EViewFlags ViewFlags, TWeakPtr<TMap<FTopLevelAssetPath, UE::FInstanceDataTransformSet>> StagedTransforms);
-	~FInstanceDataObjectFixupPanel();
-	
+	FInstanceDataObjectFixupPanel(const TObjectPtr<UObject>& InstanceDataObject, EViewFlags ViewFlags, const TWeakPtr<TMap<FTopLevelAssetPath, UE::FInstanceDataTransformSet>>& StagedTransforms);
+
+	void RefreshDetailsView();
 	TSharedPtr<IDetailsView>& GenerateDetailsView(bool bScrollbarOnLeft = false);
 	
 	void SetDiffAgainstLeft(const TSharedPtr<FAsyncDetailViewDiff>& DiffAgainstLeft);
@@ -37,9 +37,7 @@ public:
 	TSharedPtr<FAsyncDetailViewDiff> GetDiffAgainstLeft() const;
 	TSharedPtr<FAsyncDetailViewDiff> GetDiffAgainstRight() const;
 
-	bool ShouldSplitterIgnoreRow(const TWeakPtr<FDetailTreeNode>& DetailTreeNode) const;
 	bool AreAllConflictsRedirected() const;
-	const FPropertyPath& GetOriginalPath(const FPropertyPath& Path) const;
 	
 	void MarkForDelete(const FPropertyPath& Path);
 	// pass-by-copy version for delegates. Use MarkForDelete when possible
@@ -59,26 +57,16 @@ public:
 private:
 	friend class FInstanceDataObjectFixupSpecification; // for access to Redirects
 	friend class FInstanceDataObjectNameWidgetOverride;
-	friend class UInstanceDataObjectFixupUndoHandler;
-	struct FRevertInfo
-	{
-		TArray<uint8> OriginalValue;
-		FPropertyPath OriginalPath;
-		bool bHadSkipSerialization;
-		bool bWasHidden;
-	};
 
 	UE::FInstanceDataTransformSet* GetStagedTransformsForStruct(UStruct* ClassOrStruct);
 
 	void RedirectProperty(const FPropertyPath& From, const FPropertyPath& To);
+
 	// pass-by-copy version for delegates. Use RedirectProperty when possible
 	void OnRedirectProperty(FPropertyPath From, FPropertyPath To);
 
-	TObjectPtr<UObject> InstanceDataObject; // stores either InstanceDataObject property bag or just regular objects
+	TObjectPtr<UObject> InstanceDataObject;
 
-	TMap<FPropertyPath, FRevertInfo> RevertInfo;
-	TSet<FPropertyPath> MarkedForDelete;
-	
 	TWeakPtr<FAsyncDetailViewDiff> DiffAgainstLeft;
 	TWeakPtr<FAsyncDetailViewDiff> DiffAgainstRight;
 	EViewFlags ViewFlags = EViewFlags::None;
diff --git a/Engine/Plugins/Experimental/InstanceDataObjectFixupTool/Source/Private/InstanceDataObjectFixupTool.cpp b/Engine/Plugins/Experimental/InstanceDataObjectFixupTool/Source/Private/InstanceDataObjectFixupTool.cpp
--- a/Engine/Plugins/Experimental/InstanceDataObjectFixupTool/Source/Private/InstanceDataObjectFixupTool.cpp
+++ b/Engine/Plugins/Experimental/InstanceDataObjectFixupTool/Source/Private/InstanceDataObjectFixupTool.cpp
@@ -9,9 +9,11 @@
 #include "SDetailsSplitter.h"
 #include "Widgets/Layout/LinkableScrollBar.h"
 #include "Widgets/Input/SButton.h"
+#include "UObject/PropertyBagRepository.h"
 #include "DataRecoveryToolStyleSheet.h"
 #include "DataRecoveryToolUtils.h"
 #include "Widgets/Layout/SScrollBox.h"
+#include "Engine/Engine.h"
 
 #define LOCTEXT_NAMESPACE "InstanceDataObjectFixupTool"
 
@@ -76,102 +78,181 @@ void SInstanceDataObjectFixupTool::Construct(const FArguments& InArgs)
 
 	ChildSlot
 	[
-		SNew(SVerticalBox)
+		SAssignNew(Border, SBorder)
+			.BorderImage(FAppStyle::Get().GetBrush("Brushes.Recessed"))
+			.Padding(Padding::None)
+	];
+}
+
+void SInstanceDataObjectFixupTool::Tick(const FGeometry& AllottedGeometry, const double InCurrentTime, const float InDeltaTime)
+{
+	if (PanelDiff != nullptr)
+	{ 
+		PanelDiff->Tick();
+	}
+}
+
+void SInstanceDataObjectFixupTool::SetDockTab(const TSharedRef<SDockTab>& DockTab)
+{
+	OwningDockTab = DockTab;
+}
+
+void SInstanceDataObjectFixupTool::ResetDiff()
+{
+	Border->SetContent(SNullWidget::NullWidget);
+	PanelDiff.Reset();
+}
+
+TObjectPtr<UObject> SInstanceDataObjectFixupTool::CreateInstanceDataObjectSnapshot()
+{
+	const TSharedPtr<FTopLevelAssetPath> ClassPath = SelectedClassPath.Pin();
+
+	if (!ClassPath)
+	{
+		return nullptr;
+	}
+
+	TObjectPtr<UObject> Owner;
+
+	if (UE::DataRecoveryTool::Utils::GetInstanceDataObjectFromSelectedClassPath(*ClassPath, &Owner) == nullptr)
+	{
+		return nullptr;
+	}
+
+	// CPFUO the instance
+	constexpr EObjectFlags FlagsToRemove = RF_Public | RF_Standalone;
+	constexpr EObjectFlags FlagsToAdd = RF_Transient | RF_DuplicateTransient;
+
+	constexpr EObjectFlags Flags = FlagsToAdd & ~FlagsToRemove;
+
+	UObject* OwnerSnapshot = NewObject<UObject>(GetTransientPackage(), Owner->GetClass(), NAME_None, Flags);
+
+	TMap<UObject*, UObject*> ReplacementMap;
+	ReplacementMap.Add(Owner, OwnerSnapshot);
+
+	UEngine::FCopyPropertiesForUnrelatedObjectsParams Params;
+
+	Params.OptionalReplacementMappings = &ReplacementMap;
+
+	Params.bPerformDuplication = true;
+	Params.bOnlyHandleDirectSubObjects = false;
+	Params.bReplaceInternalReferenceUponRead = true;
+
+	UEngine::CopyPropertiesForUnrelatedObjects(Owner, OwnerSnapshot, Params);
+
+	OwnerSnapshot->ClearFlags(FlagsToRemove);
+	OwnerSnapshot->SetFlags(FlagsToAdd);
+
+	return UE::FPropertyBagRepository::Get().FindInstanceDataObject(OwnerSnapshot);
+}
+
+TSharedRef<SWidget> SInstanceDataObjectFixupTool::MakeDiffViewPanel()
+{
+	return SNew(SVerticalBox)
 		+ SVerticalBox::Slot()
 		.Padding(Padding::None, Padding::Small)
 		.FillHeight(1.f)
 		[
 			SNew(SScrollBox)
 			.Orientation(EOrientation::Orient_Vertical)
-			+SScrollBox::Slot()
+			+ SScrollBox::Slot()
 			[
-				SAssignNew(Border, SBorder)
-				.BorderImage(FAppStyle::Get().GetBrush("Brushes.Recessed"))
-				.Padding(Padding::None)
+				SAssignNew(Splitter, SDetailsSplitter)
+				.RowHighlightColor_Static(&SInstanceDataObjectFixupTool::GetRowHighlightColor)
 			]
 		]
-		+SVerticalBox::Slot()
+		+ SVerticalBox::Slot()
 		.AutoHeight()
 		.HAlign(HAlign_Left)
 		.VAlign(VAlign_Center)
 		.Padding(Padding::Small)
 		[
-			SNew(SButton)
-			.HAlign(HAlign_Center)
-			.VAlign(VAlign_Center)
-			.ContentPadding(Padding::Normal)
-			.ButtonStyle(FAppStyle::Get(), "SimpleButton")
-			.ToolTipText(LOCTEXT("AutoMarkForDeletion_Tooltip", "Mark remaining conflicts for deletion"))
-			.OnClicked(this, &SInstanceDataObjectFixupTool::OnAutoMarkForDeletion)
+			SNew(SHorizontalBox)
+			+ SHorizontalBox::Slot()
+			.AutoWidth()
 			[
-				SNew(SHorizontalBox)
-				+SHorizontalBox::Slot()
-				.AutoWidth()
+				SNew(SButton)
+				.HAlign(HAlign_Center)
 				.VAlign(VAlign_Center)
+				.ContentPadding(Padding::Normal)
+				.ButtonStyle(FAppStyle::Get(), "SimpleButton")
+				.ToolTipText(LOCTEXT("ResetChangesOnInstance_Tooltip", "Undo all mappings for this class type"))
+				.OnClicked(this, &SInstanceDataObjectFixupTool::OnResetInstanceDataObjectSnapshot)
+				.IsEnabled_Raw(this, &SInstanceDataObjectFixupTool::CanResetInstanceDataObjectSnapshot)
 				[
-					SNew(SImage)
-					.Image(FAppStyle::GetBrush("Icons.Delete"))
+					SNew(SHorizontalBox)
+					+ SHorizontalBox::Slot()
+					.AutoWidth()
+					.VAlign(VAlign_Center)
+					[
+						SNew(SImage)
+						.Image(FAppStyle::GetBrush("GenericCommands.Undo"))
+					]
+					+ SHorizontalBox::Slot()
+					.AutoWidth()
+					.Padding(Padding::Small, Padding::Small, Padding::Small, Padding::Small)
+					.VAlign(VAlign_Center)
+					[
+						SNew(STextBlock)
+						.Text(LOCTEXT("ResetChangesOnInstance", "Reset mappings for this class"))
+					]
 				]
-				+ SHorizontalBox::Slot()
-				.AutoWidth()
-				.Padding(Padding::Small, Padding::Small, Padding::Small, Padding::Small)
+			]
+			+ SHorizontalBox::Slot()
+			.AutoWidth()
+			[
+				SNew(SButton)
+				.HAlign(HAlign_Center)
 				.VAlign(VAlign_Center)
+				.ContentPadding(Padding::Normal)
+				.ButtonStyle(FAppStyle::Get(), "SimpleButton")
+				.ToolTipText(LOCTEXT("AutoMarkForDeletion_Tooltip", "Mark remaining conflicts for deletion"))
+				.OnClicked(this, &SInstanceDataObjectFixupTool::OnAutoMarkForDeletion)
+				.IsEnabled_Raw(this, &SInstanceDataObjectFixupTool::CanAutoMarkForDeletion)
 				[
-					SNew(STextBlock)
-					.Text(LOCTEXT("AutoMarkForDeletion", "Delete remaining conflicts"))
+					SNew(SHorizontalBox)
+					+ SHorizontalBox::Slot()
+					.AutoWidth()
+					.VAlign(VAlign_Center)
+					[
+						SNew(SImage)
+						.Image(FAppStyle::GetBrush("Icons.Delete"))
+					]
+					+ SHorizontalBox::Slot()
+					.AutoWidth()
+					.Padding(Padding::Small, Padding::Small, Padding::Small, Padding::Small)
+					.VAlign(VAlign_Center)
+					[
+						SNew(STextBlock)
+						.Text(LOCTEXT("AutoMarkForDeletion", "Delete remaining conflicts"))
+					]
 				]
 			]
-		]
-	];
-
-	// TODO: Could remove this for optimizations
-	GenerateDetailsViews();
-}
-
-void SInstanceDataObjectFixupTool::Tick(const FGeometry& AllottedGeometry, const double InCurrentTime, const float InDeltaTime)
-{
-	if (PanelDiff != nullptr)
-	{ 
-		PanelDiff->Tick();
-	}
-}
-
-END_SLATE_FUNCTION_BUILD_OPTIMIZATION
-
-void SInstanceDataObjectFixupTool::SetDockTab(const TSharedRef<SDockTab>& DockTab)
-{
-	OwningDockTab = DockTab;
-}
-
-void SInstanceDataObjectFixupTool::ResetDiff()
-{
-	Border->SetContent(SNullWidget::NullWidget);
-	PanelDiff.Reset();
+		];
 }
 
 void SInstanceDataObjectFixupTool::GenerateDetailsViews()
 {
-	TSharedPtr<FTopLevelAssetPath> ClassPath = SelectedClassPath.Pin();
+	TObjectPtr<UObject> InstanceDataObjectSnapshot = CreateInstanceDataObjectSnapshot();
 
-	if (!ClassPath)
+	if (InstanceDataObjectSnapshot == nullptr)
 	{
 		ResetDiff();
 		return;
 	}
 
-	TObjectPtr<UObject> InstanceDataObject = UE::DataRecoveryTool::Utils::GetInstanceDataObjectFromSelectedClassPath(*ClassPath);
-
-	if (InstanceDataObject == nullptr)
+	// apply staged transforms if any
+	if (TSharedPtr Transforms = StagedTransforms.Pin(); const UE::FInstanceDataTransformSet* TransformSet = Transforms->Find(*SelectedClassPath.Pin()))
 	{
-		ResetDiff();
-		return;
+		const UObject* Owner = UE::FPropertyBagRepository::Get().FindInstanceForDataObject(InstanceDataObjectSnapshot);
+		InstanceDataObjectSnapshot = UE::FInstanceDataTransforms::Get().ApplyTransformSet(*TransformSet, InstanceDataObjectSnapshot, const_cast<UObject*>(Owner));
 	}
 
 	constexpr int32 LeftIndex = 0;
 	constexpr int32 RightIndex = 1;
 
-	Panels[LeftIndex] = MakeShared<FInstanceDataObjectFixupPanel>(InstanceDataObject, FInstanceDataObjectFixupPanel::EViewFlags::DefaultLeftPanel, StagedTransforms);
-	Panels[RightIndex] = MakeShared<FInstanceDataObjectFixupPanel>(InstanceDataObject, FInstanceDataObjectFixupPanel::EViewFlags::DefaultRightPanel, StagedTransforms);
+	Panels[LeftIndex] = MakeShared<FInstanceDataObjectFixupPanel>(InstanceDataObjectSnapshot, FInstanceDataObjectFixupPanel::EViewFlags::DefaultLeftPanel, StagedTransforms);
+	Panels[RightIndex] = MakeShared<FInstanceDataObjectFixupPanel>(InstanceDataObjectSnapshot, FInstanceDataObjectFixupPanel::EViewFlags::DefaultRightPanel, StagedTransforms);
 
 	// generate the panels
 	Panels[LeftIndex]->GenerateDetailsView(true);
@@ -189,8 +270,7 @@ void SInstanceDataObjectFixupTool::GenerateDetailsViews()
 
 
 	Border->SetContent(
-		SAssignNew(Splitter, SDetailsSplitter)
-		.RowHighlightColor_Static(&SInstanceDataObjectFixupTool::GetRowHighlightColor)
+		MakeDiffViewPanel()
 	);
 
 	// add the panels to the splitter
@@ -200,12 +280,10 @@ void SInstanceDataObjectFixupTool::GenerateDetailsViews()
 			SDetailsSplitter::Slot()
 			.DetailsView(Panel->DetailsView.ToSharedRef())
 			.DifferencesWithRightPanel(Panel.Get(), &FInstanceDataObjectFixupPanel::GetDiffAgainstRight)
-			.IsReadonly_Lambda([Panel = Panel.ToWeakPtr()]() { return Panel.IsValid() ? 
+			.IsReadonly_Lambda([Panel = Panel.ToWeakPtr()]() { return Panel.IsValid() ?
 				Panel.Pin()->HasViewFlag(FInstanceDataObjectFixupPanel::EViewFlags::ReadonlyValues) : true; })
-			.ShouldIgnoreRow(Panel.Get(), &FInstanceDataObjectFixupPanel::ShouldSplitterIgnoreRow)
 		);
 	}
-	
 }
 
 	
@@ -234,6 +312,11 @@ bool SInstanceDataObjectFixupTool::IsResolved() const
 {
 	for (const TSharedPtr<FInstanceDataObjectFixupPanel>& Panel : Panels)
 	{
+		if (Panel == nullptr)
+		{
+			continue;
+		}
+
 		if (!Panel->AreAllConflictsRedirected())
 		{
 			return false;
@@ -251,6 +334,63 @@ FReply SInstanceDataObjectFixupTool::OnAutoMarkForDeletion() const
 	return FReply::Handled();
 }
 
+bool SInstanceDataObjectFixupTool::CanAutoMarkForDeletion() const
+{
+	return !IsResolved();
+}
+
+FReply SInstanceDataObjectFixupTool::OnResetInstanceDataObjectSnapshot()
+{
+	TSharedPtr<FTopLevelAssetPath> ClassPath = SelectedClassPath.Pin();
+	if (ClassPath == nullptr)
+	{
+		return FReply::Unhandled();
+	}
+
+	UClass* LoadedClass = Cast<UClass>(StaticLoadObject(UClass::StaticClass(), nullptr, ClassPath->ToString()));
+	if (LoadedClass == nullptr)
+	{
+		return FReply::Unhandled();
+	}
+
+	UE::FInstanceDataTransformSet* TransformSet = StagedTransforms.Pin()->Find(*ClassPath);
+	if (TransformSet == nullptr)
+	{
+		return FReply::Unhandled();
+	}
+
+	*TransformSet = UE::FInstanceDataTransforms::Get().GetTransformSet(*LoadedClass);
+
+	GenerateDetailsViews();
+
+	return FReply::Handled();
+
+}
+
+bool SInstanceDataObjectFixupTool::CanResetInstanceDataObjectSnapshot() const
+{
+	const UE::FInstanceDataTransformSet* TransformSet = GetSelectedStagedTransforms();
+
+	return TransformSet && !TransformSet->Operations.IsEmpty();
+}
+
+UE::FInstanceDataTransformSet* SInstanceDataObjectFixupTool::GetSelectedStagedTransforms() const
+{
+	if (const TSharedPtr<FTopLevelAssetPath> SelectedPath = SelectedClassPath.Pin())
+	{
+		if (const TSharedPtr<TMap<FTopLevelAssetPath, UE::FInstanceDataTransformSet>> Transforms = StagedTransforms.Pin())
+		{
+			if (UE::FInstanceDataTransformSet* Found = Transforms->Find(*SelectedPath))
+			{
+				return Found;
+			}
+		}
+	}
+	return nullptr;
+}
+
+END_SLATE_FUNCTION_BUILD_OPTIMIZATION
+
 
 #undef LOCTEXT_NAMESPACE
 
diff --git a/Engine/Plugins/Experimental/InstanceDataObjectFixupTool/Source/Private/InstanceDataObjectFixupToolModule.cpp b/Engine/Plugins/Experimental/InstanceDataObjectFixupTool/Source/Private/InstanceDataObjectFixupToolModule.cpp
--- a/Engine/Plugins/Experimental/InstanceDataObjectFixupTool/Source/Private/InstanceDataObjectFixupToolModule.cpp
+++ b/Engine/Plugins/Experimental/InstanceDataObjectFixupTool/Source/Private/InstanceDataObjectFixupToolModule.cpp
@@ -4,7 +4,6 @@
 
 #include "DataRecoveryTool.h"
 #include "DataRecoveryToolStyleSheet.h"
-#include "InstanceDataObjectFixupTool.h"
 #include "WorkspaceMenuStructure.h"
 #include "WorkspaceMenuStructureModule.h"
 #include "Elements/Framework/TypedElementRegistry.h"
@@ -14,11 +13,50 @@
 
 #define LOCTEXT_NAMESPACE "InstanceDataObjectFixupToolModule"
 
+static const TCHAR* CVarEnableDataRecoveryToolName = TEXT("IDO.EnableDataRecoveryTool");
+
 void FInstanceDataObjectFixupToolModule::StartupModule()
 {
+	if (!FModuleManager::Get().IsModuleLoaded("EntityEditor"))
+	{
+		FModuleManager::Get().LoadModuleChecked("EntityEditor");
+	}
+
+	CVarEnableDataRecoveryTool = IConsoleManager::Get().FindConsoleVariable(CVarEnableDataRecoveryToolName);
+	if (!CVarEnableDataRecoveryTool)
+	{
+		return;
+	}
+
+	CVarChangedHandle =
+		CVarEnableDataRecoveryTool->AsVariable()->OnChangedDelegate()
+        .AddRaw(this, &FInstanceDataObjectFixupToolModule::OnCVarChanged);
+
+	OnCVarChanged(CVarEnableDataRecoveryTool);
 }
 
-void FInstanceDataObjectFixupToolModule::ShutdownModule()
+void FInstanceDataObjectFixupToolModule::OnCVarChanged(IConsoleVariable* Var)
+{
+	Var->GetBool() == true ? RegisterTabSpawners() : UnregisterTabSpawners();
+}
+
+void FInstanceDataObjectFixupToolModule::RegisterTabSpawners(const TOptional<FTopLevelAssetPath>& ClassPath) const
+{
+	using namespace UE::DataRecoveryTool::StyleSheet;
+
+	FGlobalTabmanager::Get()->RegisterNomadTabSpawner(
+		Dialog::DataRecoveryToolDialogName,
+		FOnSpawnTab::CreateLambda([this, ClassPath](const FSpawnTabArgs& Args)
+			{
+				return CreateDataRecoveryToolTab(ClassPath);
+			}))
+		.SetGroup(WorkspaceMenu::GetMenuStructure().GetToolsCategory())
+		.SetDisplayName(LOCTEXT("DataRecoveryTool_DisplayName", "Data Recovery Tool"))
+		.SetTooltipText(LOCTEXT("DataRecoveryTool_DisplayNameToolTip", "Opens the Data recovery Tool"))
+		.SetIcon(FSlateIcon(FAppStyle::GetAppStyleSetName(), Tab::Icon));
+}
+
+void FInstanceDataObjectFixupToolModule::UnregisterTabSpawners() const
 {
 	using namespace UE::DataRecoveryTool::StyleSheet;
 
@@ -28,7 +66,19 @@ void FInstanceDataObjectFixupToolModule::ShutdownModule()
 	}
 }
 
-TSharedRef<SDockTab> FInstanceDataObjectFixupToolModule::CreateDataRecoveryToolTab(TOptional<FTopLevelAssetPath> ClassPath) const
+void FInstanceDataObjectFixupToolModule::ShutdownModule()
+{
+	UnregisterTabSpawners();
+	if (!CVarEnableDataRecoveryTool)
+	{
+		CVarEnableDataRecoveryTool->AsVariable()->OnChangedDelegate().Remove(CVarChangedHandle);
+		CVarEnableDataRecoveryTool = nullptr;
+	}
+	
+	CVarChangedHandle.Reset();
+}
+
+TSharedRef<SDockTab> FInstanceDataObjectFixupToolModule::CreateDataRecoveryToolTab(const TOptional<FTopLevelAssetPath>& ClassPath) const
 {
 	const TSharedRef<SDataRecoveryTool> DataRecoveryTool = SNew(SDataRecoveryTool)
 		.SelectedClassPath(ClassPath);
@@ -44,21 +94,11 @@ TSharedRef<SDockTab> FInstanceDataObjectFixupToolModule::CreateDataRecoveryToolT
 	return DockTab;
 }
 
-void FInstanceDataObjectFixupToolModule::CreateDataRecoveryToolDialog(const UObject* Instance, const UObject* /*Owner = nullptr*/) const
+void FInstanceDataObjectFixupToolModule::CreateDataRecoveryToolDialog(const TOptional<FTopLevelAssetPath>& ClassPath) const
 {
 	using namespace UE::DataRecoveryTool::StyleSheet;
 
-	FGlobalTabmanager::Get()->RegisterNomadTabSpawner(Dialog::DataRecoveryToolDialogName, FOnSpawnTab::CreateLambda(
-		[Instance](const FSpawnTabArgs& TabArgs)
-		{
-			TOptional<FTopLevelAssetPath> ClassPath = Instance ? Instance->GetClass()->GetClassPathName() : TOptional<FTopLevelAssetPath>{};
-
-			return Get().CreateDataRecoveryToolTab(ClassPath);
-		}))
-		.SetIcon(FSlateIcon(FAppStyle::GetAppStyleSetName(), Tab::Icon))
-		.SetDisplayName(Tab::DisplayName)
-		.SetGroup(WorkspaceMenu::GetMenuStructure().GetToolsCategory())
-		.SetMenuType(ETabSpawnerMenuType::Hidden);
+	RegisterTabSpawners(ClassPath);
 
 	if (const TSharedPtr<SDockTab> DockTab = FGlobalTabmanager::Get()->TryInvokeTab(FTabId(Dialog::DataRecoveryToolDialogName)))
 	{
diff --git a/Engine/Plugins/Experimental/InstanceDataObjectFixupTool/Source/Private/InstanceDataObjectFixupToolTedsQueries.cpp b/Engine/Plugins/Experimental/InstanceDataObjectFixupTool/Source/Private/InstanceDataObjectFixupToolTedsQueries.cpp
--- a/Engine/Plugins/Experimental/InstanceDataObjectFixupTool/Source/Private/InstanceDataObjectFixupToolTedsQueries.cpp
+++ b/Engine/Plugins/Experimental/InstanceDataObjectFixupTool/Source/Private/InstanceDataObjectFixupToolTedsQueries.cpp
@@ -74,14 +74,17 @@ void UInstanceDataObjectFixupToolTedsQueryFactory::ShowFixUpTool(UE::Editor::Dat
 			UE::FPropertyBagRepository::Get().FindNestedInstanceDataObject(Owner, true,
 				[Owner](UObject* NestedObject)
 				{
-					FInstanceDataObjectFixupToolModule::Get().CreateDataRecoveryToolDialog(NestedObject, Owner);
+					FTopLevelAssetPath ClassPath = Owner->GetClass()->GetClassPathName();
+					FInstanceDataObjectFixupToolModule::Get().CreateDataRecoveryToolDialog(ClassPath);
 				});
 		}
 		else
 		{
-			if (TObjectPtr<UObject> InstanceDataObject = UE::FPropertyBagRepository::Get().FindInstanceDataObject(ObjectColumn->Object.Get()))
+			UObject* Owner = ObjectColumn->Object.Get();
+			if (TObjectPtr<UObject> InstanceDataObject = UE::FPropertyBagRepository::Get().FindInstanceDataObject(Owner))
 			{
-				FInstanceDataObjectFixupToolModule::Get().CreateDataRecoveryToolDialog(InstanceDataObject);
+				FTopLevelAssetPath ClassPath = Owner->GetClass()->GetClassPathName();
+				FInstanceDataObjectFixupToolModule::Get().CreateDataRecoveryToolDialog(ClassPath);
 			}
 		}
 	}
diff --git a/Engine/Plugins/Experimental/InstanceDataObjectFixupTool/Source/Public/DataRecoveryTool.h b/Engine/Plugins/Experimental/InstanceDataObjectFixupTool/Source/Public/DataRecoveryTool.h
--- a/Engine/Plugins/Experimental/InstanceDataObjectFixupTool/Source/Public/DataRecoveryTool.h
+++ b/Engine/Plugins/Experimental/InstanceDataObjectFixupTool/Source/Public/DataRecoveryTool.h
@@ -8,12 +8,22 @@
 #include "Widgets/Docking/SDockTab.h"
 #include "Widgets/SCompoundWidget.h"
 #include "DataRecoveryToolUtils.h"
+#include "UObject/InstanceDataTransforms.h"
+
 
 #define UE_API INSTANCEDATAOBJECTFIXUPTOOL_API
 
+
 class SInstanceDataObjectFixupTool;
+
 namespace UE
 {
+	namespace Editor::DataStorage::QueryStack
+	{
+		class FRowQueryResultsNode;
+	}
+
+	class FRowQueryResultsNode;
 	class FInstanceDataTransformSet;
 }
 
@@ -42,6 +52,8 @@ private:
 	UE_API TSharedRef<SWidget> MakeTEDSTableViewer();
 	UE_API void OnListSelectionChanged(UE::Editor::DataStorage::RowHandle RowItem);
 
+	void SetSelectedClass();
+
 	UE_API TSharedRef<SWidget> MakeRightPanel();
 
 	UE_API TSharedRef<SWidget> MakeDiffSection();
@@ -49,14 +61,16 @@ private:
 	FReply OnApplyClicked() const;
 	FReply OnSaveAndApplyClicked() const;
 
-
 	TWeakPtr<SDockTab> OwningDockTab;
 
 	TSharedPtr<FTopLevelAssetPath> SelectedClassPath;
+
 	TSharedPtr<TMap<FTopLevelAssetPath, UE::FInstanceDataTransformSet>> StagedTransforms;
 
 	TSharedPtr<UE::Editor::DataStorage::STedsTableViewer> ClassesTableViewer;
 
+	TSharedPtr<UE::Editor::DataStorage::QueryStack::FRowQueryResultsNode> ReferenceQueryResultsNode;
+
 	TSharedPtr<SInstanceDataObjectFixupTool> DiffView;
 };
 
diff --git a/Engine/Plugins/Experimental/InstanceDataObjectFixupTool/Source/Public/InstanceDataObjectFixupTool.h b/Engine/Plugins/Experimental/InstanceDataObjectFixupTool/Source/Public/InstanceDataObjectFixupTool.h
--- a/Engine/Plugins/Experimental/InstanceDataObjectFixupTool/Source/Public/InstanceDataObjectFixupTool.h
+++ b/Engine/Plugins/Experimental/InstanceDataObjectFixupTool/Source/Public/InstanceDataObjectFixupTool.h
@@ -43,10 +43,20 @@ public:
 	UE_API void GenerateDetailsViews();
 	static UE_API FLinearColor GetRowHighlightColor(const TUniquePtr<FAsyncDetailViewDiff::DiffNodeType>& DiffNode);
 	UE_API bool IsResolved() const;
-	UE_API FReply OnAutoMarkForDeletion() const;
 	
 private:
 	void ResetDiff();
+	TObjectPtr<UObject> CreateInstanceDataObjectSnapshot();
+
+	FReply OnAutoMarkForDeletion() const;
+	bool CanAutoMarkForDeletion() const;
+
+	TSharedRef<SWidget> MakeDiffViewPanel();
+
+	FReply OnResetInstanceDataObjectSnapshot();
+	bool CanResetInstanceDataObjectSnapshot() const;
+
+	UE::FInstanceDataTransformSet* GetSelectedStagedTransforms() const;
 
 	TWeakPtr<SDockTab> OwningDockTab;
 	TStaticArray<TSharedPtr<FInstanceDataObjectFixupPanel>, 2> Panels;
@@ -58,6 +68,8 @@ private:
 
 	TSharedPtr<SBorder> Border;
 	TWeakPtr<FTopLevelAssetPath> SelectedClassPath;
+
+	
 };
 
 #undef UE_API
diff --git a/Engine/Plugins/Experimental/InstanceDataObjectFixupTool/Source/Public/InstanceDataObjectFixupToolModule.h b/Engine/Plugins/Experimental/InstanceDataObjectFixupTool/Source/Public/InstanceDataObjectFixupToolModule.h
--- a/Engine/Plugins/Experimental/InstanceDataObjectFixupTool/Source/Public/InstanceDataObjectFixupToolModule.h
+++ b/Engine/Plugins/Experimental/InstanceDataObjectFixupTool/Source/Public/InstanceDataObjectFixupToolModule.h
@@ -16,10 +16,18 @@ public:
 	{
 		return FModuleManager::LoadModuleChecked<FInstanceDataObjectFixupToolModule>("InstanceDataObjectFixupTool");
 	}
-	
-	UE_API TSharedRef<SDockTab> CreateDataRecoveryToolTab(TOptional<FTopLevelAssetPath> ClassPath = {}) const;
 
-	UE_API void CreateDataRecoveryToolDialog(const UObject* Instance, const UObject* Owner = nullptr) const;
+	UE_API TSharedRef<SDockTab> CreateDataRecoveryToolTab(const TOptional<FTopLevelAssetPath>& ClassPath = {}) const;
+
+	UE_API void CreateDataRecoveryToolDialog(const TOptional<FTopLevelAssetPath>& ClassPath) const;
+
+private:
+	void OnCVarChanged(IConsoleVariable* Var);
+	void RegisterTabSpawners(const TOptional<FTopLevelAssetPath>& ClassPath = {}) const;
+	void UnregisterTabSpawners() const;
+
+	IConsoleVariable* CVarEnableDataRecoveryTool;
+	FDelegateHandle CVarChangedHandle;
 };
 
 #undef UE_API
