[BugFix] Dereferencing a nullptr warning fixup

diff --git a/Engine/Plugins/Experimental/InstanceDataObjectFixupTool/Source/Private/InstanceDataObjectFixupToolModule.cpp b/Engine/Plugins/Experimental/InstanceDataObjectFixupTool/Source/Private/InstanceDataObjectFixupToolModule.cpp
--- a/Engine/Plugins/Experimental/InstanceDataObjectFixupTool/Source/Private/InstanceDataObjectFixupToolModule.cpp
+++ b/Engine/Plugins/Experimental/InstanceDataObjectFixupTool/Source/Private/InstanceDataObjectFixupToolModule.cpp
@@ -69,12 +69,11 @@ void FInstanceDataObjectFixupToolModule::UnregisterTabSpawners() const
 void FInstanceDataObjectFixupToolModule::ShutdownModule()
 {
 	UnregisterTabSpawners();
-	if (!CVarEnableDataRecoveryTool)
+	if (CVarEnableDataRecoveryTool)
 	{
 		CVarEnableDataRecoveryTool->AsVariable()->OnChangedDelegate().Remove(CVarChangedHandle);
 		CVarEnableDataRecoveryTool = nullptr;
 	}
-	
 	CVarChangedHandle.Reset();
 }
 
