[Feature] DRT Alerts & refactors

- Added Alerts columns in the class selection panel
- Alerts displays how many instances need fixup for that specific type
- Refactored how the IDTs are applied on the objects
- Macro to toggle if the real instance or snapshot is used in the diff view

diff --git a/Engine/Plugins/Experimental/EditorDataStorageFeatures/Source/TedsTypeInfo/Private/TedsTypeInfoFactory.cpp b/Engine/Plugins/Experimental/EditorDataStorageFeatures/Source/TedsTypeInfo/Private/TedsTypeInfoFactory.cpp
--- a/Engine/Plugins/Experimental/EditorDataStorageFeatures/Source/TedsTypeInfo/Private/TedsTypeInfoFactory.cpp
+++ b/Engine/Plugins/Experimental/EditorDataStorageFeatures/Source/TedsTypeInfo/Private/TedsTypeInfoFactory.cpp
@@ -14,6 +14,8 @@
 #include "Elements/Interfaces/TypedElementDataStorageInterface.h"
 #include "Elements/Columns/TypedElementLabelColumns.h"
 #include "Elements/Framework/TypedElementQueryBuilder.h"
+#include "TedsAlerts.h"
+#include "TedsAlertColumns.h"
 
 #include "HAL/IConsoleManager.h"
 #include "UObject/InstanceDataObjectUtils.h"
@@ -279,12 +281,37 @@ void UTypeInfoFactory::AddIdoColumns(UE::Editor::DataStorage::ICoreProvider& Dat
 {
 	using namespace UE::Editor::DataStorage;
 
+
 	checkf(InVerseTypeInfo, TEXT("Attempted to add type info columns with null VerseTypeInfo"));
 
 	// Try find an InstanceDataObject from type
 	if (UE::FPropertyBagRepository::Get().FindFirstInstanceThatNeedsFixup(InVerseTypeInfo) != nullptr)
 	{
 		DataStorage.AddColumn<FHasUnknownPropertiesTag>(TypeRowHandle);
+
+		TArray<UObject*> AllInstances;
+		GetObjectsOfClass(InVerseTypeInfo, AllInstances, /*bIncludeDerivedClasses =*/ false);
+
+		uint8 InstanceWithUnknownPropertiesCount = 0;
+		for (UObject* Instance : AllInstances)
+		{
+			if (UE::FPropertyBagRepository::Get().RequiresFixup(Instance))
+			{
+				++InstanceWithUnknownPropertiesCount;
+			}
+		}
+
+		const FText AlertMessagePattern = NSLOCTEXT("InstanceDataObjectTypes", "UnknownPropertiesAlertMessage",
+			"Unknown properties found in {InstanceCount} {InstanceCount}|plural(one=instance, other=instances).");
+
+		FFormatNamedArguments AlertMessageArgs;
+		AlertMessageArgs.Add(TEXT("InstanceCount"), InstanceWithUnknownPropertiesCount);
+
+		const FText AlertMessage = FText::Format(AlertMessagePattern, AlertMessageArgs);
+
+		static FName AlertName("UnknownPropertiesAlert");
+
+		Alerts::AddAlert(DataStorage, TypeRowHandle, AlertName, AlertMessage, Columns::FAlertColumnType::Warning);
 	}
 }
 
diff --git a/Engine/Plugins/Experimental/EditorDataStorageFeatures/Source/TedsTypeInfo/TedsTypeInfo.build.cs b/Engine/Plugins/Experimental/EditorDataStorageFeatures/Source/TedsTypeInfo/TedsTypeInfo.build.cs
--- a/Engine/Plugins/Experimental/EditorDataStorageFeatures/Source/TedsTypeInfo/TedsTypeInfo.build.cs
+++ b/Engine/Plugins/Experimental/EditorDataStorageFeatures/Source/TedsTypeInfo/TedsTypeInfo.build.cs
@@ -28,6 +28,7 @@ public class TedsTypeInfo : ModuleRules
 					"Engine",
 					"TypedElementFramework",
 					"UnrealEd",
+					"TedsAlerts",
 				});
 
 			PublicDependencyModuleNames.AddRange(
diff --git a/Engine/Plugins/Experimental/InstanceDataObjectFixupTool/Source/Private/DataRecoveryTool.cpp b/Engine/Plugins/Experimental/InstanceDataObjectFixupTool/Source/Private/DataRecoveryTool.cpp
--- a/Engine/Plugins/Experimental/InstanceDataObjectFixupTool/Source/Private/DataRecoveryTool.cpp
+++ b/Engine/Plugins/Experimental/InstanceDataObjectFixupTool/Source/Private/DataRecoveryTool.cpp
@@ -2,6 +2,8 @@
 
 #include "DataRecoveryTool.h"
 
+
+
 #include "DataRecoveryToolStyleSheet.h"
 #include "DataStorage/Features.h"
 #include "DetailTreeNode.h"
@@ -25,15 +27,24 @@
 #include "TedsTableViewer/Public/Widgets/STedsTableViewer.h"
 #include "UObject/InstanceDataTransforms.h"
 #include "Widgets/Input/SButton.h"
-#include "Widgets/Layout/SSeparator.h"
 #include "Widgets/STedsHierarchyViewer.h"
 #include "Serialization/ObjectWriter.h"
 #include "Misc/ScopedSlowTask.h"
 
+
+#define USE_SAVE_BUTTON 0
+#if USE_SAVE_BUTTON
+#include "Widgets/Layout/SSeparator.h"
+#endif // USE_SAVE_BUTTON
+
+#define USE_OPEN_EXTERNAL_BUTTON 0
+
 #define LOCTEXT_NAMESPACE "DataRecoveryTool"
 
 BEGIN_SLATE_FUNCTION_BUILD_OPTIMIZATION
 
+
+
 using namespace UE::DataRecoveryTool::StyleSheet;
 
 void SDataRecoveryTool::Construct(const FArguments& InArgs)
@@ -78,6 +89,7 @@ TSharedRef<SWidget> SDataRecoveryTool::MakeTopBar()
 		.Padding(Padding::None)
 		[
 			SNew(SHorizontalBox)
+#if USE_SAVE_BUTTON
 			+SHorizontalBox::Slot()
 			.AutoWidth()
 			.Padding(0)
@@ -95,66 +107,67 @@ TSharedRef<SWidget> SDataRecoveryTool::MakeTopBar()
 				]
 			]
 			+ SHorizontalBox::Slot()
-				.AutoWidth()
-				.VAlign(VAlign_Fill)
-				[
-					SNew(SSeparator)
-					.Thickness(Padding::Small * .5f)
-					.Orientation(Orient_Vertical)
-					.SeparatorImage(FAppStyle::Get().GetBrush("Brushes.Background"))
-				]
+			.AutoWidth()
+			.VAlign(VAlign_Fill)
+			[
+				SNew(SSeparator)
+				.Thickness(Padding::Small * .5f)
+				.Orientation(Orient_Vertical)
+				.SeparatorImage(FAppStyle::Get().GetBrush("Brushes.Background"))
+			]
+#endif // USE_SAVE_BUTTON
+			+ SHorizontalBox::Slot()
+			.FillWidth(1.f)
+			[
+				SNew(SHorizontalBox)
 				+ SHorizontalBox::Slot()
-				.FillWidth(1.f)
+				.AutoWidth()
+				.HAlign(HAlign_Left)
 				[
-					SNew(SHorizontalBox)
-					+ SHorizontalBox::Slot()
-					.AutoWidth()
-					.HAlign(HAlign_Left)
+					SNew(SButton)
+					.HAlign(HAlign_Center)
+					.VAlign(VAlign_Center)
+					.ContentPadding(Padding::Normal)
+					.ButtonStyle(FAppStyle::Get(), "SimpleButton")
+					.OnClicked(this, &SDataRecoveryTool::OnApplyClicked)
 					[
-						SNew(SButton)
-						.HAlign(HAlign_Center)
-						.VAlign(VAlign_Center)
-						.ContentPadding(Padding::Normal)
-						.ButtonStyle(FAppStyle::Get(), "SimpleButton")
-						.OnClicked(this, &SDataRecoveryTool::OnApplyClicked)
+						SNew(SHorizontalBox)
+						+ SHorizontalBox::Slot()
+						.AutoWidth()
 						[
-							SNew(SHorizontalBox)
-							+ SHorizontalBox::Slot()
-							.AutoWidth()
-							[
-								SNew(SImage)
-								.Image(FAppStyle::GetBrush("LevelEditor.Recompile.Small"))
-							]
-							+ SHorizontalBox::Slot()
-							.Padding(Padding::Small, Padding::None, Padding::None, Padding::None)
-							.AutoWidth()
-							[
-								SNew(STextBlock)
-								.Text(LOCTEXT("ApplyButtonText", "Apply"))
-							]
+							SNew(SImage)
+							.Image(FAppStyle::GetBrush("LevelEditor.Recompile.Small"))
+						]
+						+ SHorizontalBox::Slot()
+						.Padding(Padding::Small, Padding::None, Padding::None, Padding::None)
+						.AutoWidth()
+						[
+							SNew(STextBlock)
+							.Text(LOCTEXT("ApplyButtonText", "Apply"))
 						]
 					]
-					+ SHorizontalBox::Slot()
-					.FillWidth(1.f)
-					[
-						SNullWidget::NullWidget
-					]
-
-					+ SHorizontalBox::Slot()
-					.HAlign(HAlign_Right)
-					[
-						SNew(SButton)
-						.HAlign(HAlign_Center)
-						.VAlign(VAlign_Center)
-						.ContentPadding(Padding::Normal)
-						.ButtonStyle(FAppStyle::Get(), "SimpleButton")
+				]
+				+ SHorizontalBox::Slot()
+				.FillWidth(1.f)
+				[
+					SNullWidget::NullWidget
+				]
+#if USE_OPEN_EXTERNAL_BUTTON
+				+ SHorizontalBox::Slot()
+				.HAlign(HAlign_Right)
+				[
+					SNew(SButton)
+					.HAlign(HAlign_Center)
+					.VAlign(VAlign_Center)
+					.ContentPadding(Padding::Normal)
+					.ButtonStyle(FAppStyle::Get(), "SimpleButton")
 					[
 						SNew(SHorizontalBox)
 						+ SHorizontalBox::Slot()
 						.AutoWidth()
 						[
 							SNew(SImage)
-							.Image(GetOpenExternalImage())
+							.Image(FAppStyle::GetBrush("Icons.OpenInExternalEditor"))
 						]
 						+ SHorizontalBox::Slot()
 						.Padding(Padding::Small, Padding::None, Padding::None, Padding::None)
@@ -165,6 +178,7 @@ TSharedRef<SWidget> SDataRecoveryTool::MakeTopBar()
 						]
 					]
 				]
+#endif // USE_OPEN_EXTERNAL_BUTTON
 			]
 		];
 }
@@ -199,7 +213,7 @@ FReply SDataRecoveryTool::OnApplyClicked() const
 {
 	using namespace UE::Editor::DataStorage;
 	using namespace UE::Editor::DataStorage::Queries;
-
+	using namespace UE::DataRecoveryTool::Utils;
 	// for now just apply the transforms to all instances of every listed class. This may be slow!
 
 	ICoreProvider* Storage = GetMutableDataStorageFeature<ICoreProvider>(StorageFeatureName);
@@ -217,41 +231,40 @@ FReply SDataRecoveryTool::OnApplyClicked() const
 
 		if (const UClass* Class = TypeInfo->TypeInfo.Get())
 		{
-			bool bClassStillRequiresFixup = false;
 			const FTopLevelAssetPath ClassPath = Class->GetClassPathName();
 
 			ApplyAllTask.EnterProgressFrame(1.f, FText::Format(LOCTEXT("ApplyToSlow", "Applying transforms to {0}"), FText::FromString(ClassPath.ToString())));
 
-			const UE::FInstanceDataTransformSet* TransformSet = StagedTransforms->Find(ClassPath);
-
-			if (TransformSet == nullptr)
-			{
-				continue;
-			}
-
 			TArray<UObject*> AllInstancesOfClasses;
 
-			// Could not use ApplyTransformSet inside ForEachObjectOfClass lambda, so we have to get the objects first
+			uint8 InstanceWithUnknownPropertiesCount = 0;
+
 			GetObjectsOfClass(Class, AllInstancesOfClasses);
 
 			for (UObject* Instance : AllInstancesOfClasses)
 			{
 				if (UObject* Ido = UE::FPropertyBagRepository::Get().FindInstanceDataObject(Instance))
 				{
-					UE::FInstanceDataTransforms::Get().ApplyTransformSet(*TransformSet, Ido, Instance);
-
-					if (!bClassStillRequiresFixup && UE::FPropertyBagRepository::Get().RequiresFixup(Instance))
+					if (!Snapshot::IsASnapshot(Instance) && UE::FPropertyBagRepository::Get().RequiresFixup(Instance))
 					{
-						bClassStillRequiresFixup = true;
+						ApplyTransforms(StagedTransforms, Ido, Instance);
+
+						++InstanceWithUnknownPropertiesCount;
 					}
 				}
 			}
 
-			if (!bClassStillRequiresFixup)
+			if (InstanceWithUnknownPropertiesCount == 0)
 			{
 				Storage->RemoveColumn<FHasUnknownPropertiesTag>(Row);
 				Storage->AddColumn<FTypedElementSyncBackToWorldTag>(Row);
 			}
+			else
+			{
+				UpdateTedsIdoAlertColumn(Storage, Row, InstanceWithUnknownPropertiesCount);
+			}
+
+			DiffView->GenerateDetailsViews();
 		}
 	}
 
@@ -287,7 +300,6 @@ FReply SDataRecoveryTool::OnSaveAndApplyClicked() const
 
 		if (const UClass* Class = TypeInfo->TypeInfo.Get())
 		{
-			bool bClassStillRequiresFixup = false;
 			const FTopLevelAssetPath ClassPath = Class->GetClassPathName();
 
 			PatchAllTask.EnterProgressFrame(1.f, FText::Format(LOCTEXT("PatchToSlow", "Patching transforms to {0}"), FText::FromString(ClassPath.ToString())));
@@ -301,6 +313,8 @@ FReply SDataRecoveryTool::OnSaveAndApplyClicked() const
 
 			TArray<UObject*> AllInstancesOfClasses;
 
+			uint8 InstanceWithUnknownPropertiesCount = 0;
+
 			// Could not use PatchTransformSet inside ForEachObjectOfClass lambda, so we have to get the objects first
 			GetObjectsOfClass(Class, AllInstancesOfClasses);
 
@@ -308,20 +322,26 @@ FReply SDataRecoveryTool::OnSaveAndApplyClicked() const
 			{
 				if (UObject* Ido = UE::FPropertyBagRepository::Get().FindInstanceDataObject(Instance))
 				{
-					UE::FInstanceDataTransforms::Get().PatchInstanceDataObject(Ido, Instance);
-
-					if (!bClassStillRequiresFixup && UE::FPropertyBagRepository::Get().RequiresFixup(Instance))
+					if (!UE::DataRecoveryTool::Utils::Snapshot::IsASnapshot(Instance) && UE::FPropertyBagRepository::Get().RequiresFixup(Instance))
 					{
-						bClassStillRequiresFixup = true;
+						UE::FInstanceDataTransforms::Get().PatchInstanceDataObject(Ido, Instance);
+
+						++InstanceWithUnknownPropertiesCount;
 					}
 				}
 			}
 
-			if (!bClassStillRequiresFixup)
+			if (InstanceWithUnknownPropertiesCount == 0)
 			{
 				Storage->RemoveColumn<FHasUnknownPropertiesTag>(Row);
 				Storage->AddColumn<FTypedElementSyncBackToWorldTag>(Row);
 			}
+			else
+			{
+				UE::DataRecoveryTool::Utils::UpdateTedsIdoAlertColumn(Storage, Row, InstanceWithUnknownPropertiesCount);
+			}
+
+			DiffView->GenerateDetailsViews();
 		}
 	}
 
@@ -500,7 +520,6 @@ TSharedRef<SWidget> SDataRecoveryTool::MakeTEDSTableViewer()
 
 	ClassesTableViewer = SNew(STedsTableViewer)
 		.QueryStack(SearchNode)
-		.EmptyRowsMessage(LOCTEXT("EmptyRowsMessage", "No loose properties found"))
 		.Columns({ FTypedElementClassTypeInfoColumn::StaticStruct(), FTedsAlertColumn::StaticStruct() })
 		.ListSelectionMode(ESelectionMode::Single)
 		.OnSelectionChanged(this, &SDataRecoveryTool::OnListSelectionChanged);
@@ -540,3 +559,5 @@ END_SLATE_FUNCTION_BUILD_OPTIMIZATION
 
 #undef LOCTEXT_NAMESPACE
 
+#undef USE_OPEN_EXTERNAL_BUTTON
+#undef USE_SAVE_BUTTON
diff --git a/Engine/Plugins/Experimental/InstanceDataObjectFixupTool/Source/Private/DataRecoveryToolUtils.cpp b/Engine/Plugins/Experimental/InstanceDataObjectFixupTool/Source/Private/DataRecoveryToolUtils.cpp
new file mode 100644
--- /dev/null
+++ b/Engine/Plugins/Experimental/InstanceDataObjectFixupTool/Source/Private/DataRecoveryToolUtils.cpp
@@ -0,0 +1,73 @@
+// Copyright Epic Games, Inc. All Rights Reserved.
+
+#include "DataRecoveryToolUtils.h"
+
+#define LOCTEXT_NAMESPACE "DataRecoveryToolModule"
+
+
+TObjectPtr<UObject> UE::DataRecoveryTool::Utils::GetInstanceDataObjectFromSelectedClassPath(const FTopLevelAssetPath& ClassPath,
+	TObjectPtr<UObject>* OutOwner)
+{
+	const UClass* LoadedClass = Cast<UClass>(StaticLoadObject(UClass::StaticClass(), nullptr, ClassPath.ToString()));
+
+	const UObject* FoundInstance = FPropertyBagRepository::Get().FindFirstInstanceThatNeedsFixup(LoadedClass);
+
+	if (OutOwner)
+	{
+		*OutOwner = const_cast<UObject*>(FoundInstance);
+	}
+
+	return FoundInstance ? FPropertyBagRepository::Get().FindInstanceDataObject(FoundInstance) : nullptr;
+}
+
+void UE::DataRecoveryTool::Utils::ApplyTransforms(const TWeakPtr<TMap<FTopLevelAssetPath, FInstanceDataTransformSet>>& StagedTransformsWeak,
+	UObject*& IDO, TNotNull<UObject*> Owner)
+{
+	const TSharedPtr StagedTransforms = StagedTransformsWeak.Pin();
+
+	if (StagedTransforms == nullptr)
+	{
+		return;
+	}
+
+	for (const TPair<FTopLevelAssetPath, FInstanceDataTransformSet>& TransformPair : *StagedTransforms)
+	{
+		FInstanceDataTransformSet TransformSet = TransformPair.Value;
+		IDO = FInstanceDataTransforms::Get().ApplyTransformSet(TransformPair.Value, IDO, Owner);
+	}
+}
+
+bool UE::DataRecoveryTool::Utils::ContainsPath(const FTopLevelAssetPath& Parent, const FTopLevelAssetPath& Candidate)
+{
+	const FString ParentPackage = Parent.GetPackageName().ToString();
+	const FString CandidatePackage = Candidate.GetPackageName().ToString();
+
+	return CandidatePackage.StartsWith(ParentPackage, ESearchCase::IgnoreCase);
+}
+
+bool UE::DataRecoveryTool::Utils::Snapshot::IsASnapshot(const UObject* Instance)
+{
+	return Instance->HasAllFlags(FlagsToAdd);
+}
+
+void UE::DataRecoveryTool::Utils::UpdateTedsIdoAlertColumn(
+	Editor::DataStorage::ICoreProvider* Storage,
+	Editor::DataStorage::RowHandle Row,
+	const uint32 InstanceWithIdoCount
+)
+{
+	using namespace Editor::DataStorage;
+	const FText AlertMessagePattern = LOCTEXT("UnknownPropertiesAlertMessage",
+		"Unknown properties found in {InstanceCount} {InstanceCount}|plural(one=instance, other=instances).");
+
+	FFormatNamedArguments AlertMessageArgs;
+	AlertMessageArgs.Add(TEXT("InstanceCount"), InstanceWithIdoCount);
+
+	const FText AlertMessage = FText::Format(AlertMessagePattern, AlertMessageArgs);
+
+	static FName AlertName("UnknownPropertiesAlert");
+
+	Alerts::AddAlert(*Storage, Row, AlertName, AlertMessage, Columns::FAlertColumnType::Warning);
+}
+
+#undef LOCTEXT_NAMESPACE
\ No newline at end of file
diff --git a/Engine/Plugins/Experimental/InstanceDataObjectFixupTool/Source/Private/DataRecoveryToolUtils.h b/Engine/Plugins/Experimental/InstanceDataObjectFixupTool/Source/Private/DataRecoveryToolUtils.h
--- a/Engine/Plugins/Experimental/InstanceDataObjectFixupTool/Source/Private/DataRecoveryToolUtils.h
+++ b/Engine/Plugins/Experimental/InstanceDataObjectFixupTool/Source/Private/DataRecoveryToolUtils.h
@@ -7,22 +7,31 @@
 #include "UObject/PropertyStateTracking.h"
 #include "UObject/PropertyPathName.h"
 #include "UObject/PropertyBagRepository.h"
+#include "UObject/InstanceDataTransforms.h"
+#include "TedsAlertColumns.h"
+#include "TedsAlerts.h"
+
 
 namespace UE::DataRecoveryTool::Utils
 {
-	inline TObjectPtr<UObject> GetInstanceDataObjectFromSelectedClassPath(const FTopLevelAssetPath& ClassPath, TObjectPtr<UObject>* OutOwner = nullptr)
-	{
-		const UClass* LoadedClass = Cast<UClass>(StaticLoadObject(UClass::StaticClass(), nullptr, ClassPath.ToString()));
+	TObjectPtr<UObject> GetInstanceDataObjectFromSelectedClassPath(const FTopLevelAssetPath& ClassPath, TObjectPtr<UObject>* OutOwner = nullptr);
 
-		const UObject* FoundInstance = FPropertyBagRepository::Get().FindFirstInstanceThatNeedsFixup(LoadedClass);
+	void ApplyTransforms(const TWeakPtr<TMap<FTopLevelAssetPath, FInstanceDataTransformSet>>& StagedTransformsWeak, UObject*& IDO, TNotNull<UObject*> Owner);
 
-		if (OutOwner)
-		{
-			*OutOwner = const_cast<UObject*>(FoundInstance);
-		}
+	bool ContainsPath(const FTopLevelAssetPath& Parent, const FTopLevelAssetPath& Candidate);
 
-		return FoundInstance ? FPropertyBagRepository::Get().FindInstanceDataObject(FoundInstance) : nullptr;
+	namespace Snapshot
+	{
+		constexpr EObjectFlags FlagsToAdd = RF_Transient | RF_DuplicateTransient;
+		constexpr EObjectFlags FlagsToRemove = RF_Public | RF_Standalone;
+
+		bool IsASnapshot(const UObject* Instance);
 	}
 
-	
+	void UpdateTedsIdoAlertColumn(
+		Editor::DataStorage::ICoreProvider* Storage,
+		Editor::DataStorage::RowHandle Row,
+		const uint32 InstanceWithIdoCount
+	);
 }
+
diff --git a/Engine/Plugins/Experimental/InstanceDataObjectFixupTool/Source/Private/InstanceDataObjectFixupPanel.cpp b/Engine/Plugins/Experimental/InstanceDataObjectFixupTool/Source/Private/InstanceDataObjectFixupPanel.cpp
--- a/Engine/Plugins/Experimental/InstanceDataObjectFixupTool/Source/Private/InstanceDataObjectFixupPanel.cpp
+++ b/Engine/Plugins/Experimental/InstanceDataObjectFixupTool/Source/Private/InstanceDataObjectFixupPanel.cpp
@@ -4,6 +4,8 @@
 #include "InstanceDataObjectFixupPanel.h"
 
 #include "AsyncDetailViewDiff.h"
+#include "DataRecoveryToolStyleSheet.h"
+#include "DataRecoveryToolUtils.h"
 #include "DetailTreeNode.h"
 #include "Widgets/Layout/LinkableScrollBar.h"
 #include "InstanceDataObjectFixupDetailCustomization.h"
@@ -381,7 +383,7 @@ void FInstanceDataObjectFixupPanel::RedirectProperty(const FPropertyPath& From,
 		Transforms->Operations.Add({ .OpCode = TEXT("RedirectProperty"), .Params = Params });
 
 		TObjectPtr<UObject> OwnerObjectPtr = const_cast<UObject*>(Owner);
-		InstanceDataObject = UE::FInstanceDataTransforms::Get().ApplyTransformSet(*Transforms, InstanceDataObject, OwnerObjectPtr);
+		UE::FInstanceDataTransforms::Get().ApplyTransformSet(*Transforms, InstanceDataObject, OwnerObjectPtr);
 
 		RefreshDetailsView();
 	}
@@ -436,7 +438,7 @@ void FInstanceDataObjectFixupPanel::MarkForDelete(const FPropertyPath& CurrentPa
 		Transforms->Operations.Add({ .OpCode = TEXT("RemoveProperty"), .Params = Params });
 
 		TObjectPtr<UObject> OwnerObjectPtr = const_cast<UObject*>(Owner);
-		InstanceDataObject = UE::FInstanceDataTransforms::Get().ApplyTransformSet(*Transforms, InstanceDataObject, OwnerObjectPtr);
+		UE::FInstanceDataTransforms::Get().ApplyTransformSet(*Transforms, InstanceDataObject, OwnerObjectPtr);
 	}
 
 	RefreshDetailsView();
diff --git a/Engine/Plugins/Experimental/InstanceDataObjectFixupTool/Source/Private/InstanceDataObjectFixupTool.cpp b/Engine/Plugins/Experimental/InstanceDataObjectFixupTool/Source/Private/InstanceDataObjectFixupTool.cpp
--- a/Engine/Plugins/Experimental/InstanceDataObjectFixupTool/Source/Private/InstanceDataObjectFixupTool.cpp
+++ b/Engine/Plugins/Experimental/InstanceDataObjectFixupTool/Source/Private/InstanceDataObjectFixupTool.cpp
@@ -5,7 +5,6 @@
 #include "InstanceDataObjectFixupPanel.h"
 #include "DetailTreeNode.h"
 #include "SlateOptMacros.h"
-#include "PropertyEditorModule.h"
 #include "SDetailsSplitter.h"
 #include "Widgets/Layout/LinkableScrollBar.h"
 #include "Widgets/Input/SButton.h"
@@ -15,8 +14,9 @@
 #include "Widgets/Layout/SScrollBox.h"
 #include "Engine/Engine.h"
 
-#define LOCTEXT_NAMESPACE "InstanceDataObjectFixupTool"
+#define USE_SNAPSHOT_INSTANCE 0
 
+#define LOCTEXT_NAMESPACE "InstanceDataObjectFixupTool"
 using namespace UE::DataRecoveryTool::StyleSheet;
 
 class FInstanceDataObjectFixupSpecification : public TTreeDiffSpecification<TWeakPtr<FDetailTreeNode>>
@@ -105,6 +105,10 @@ void SInstanceDataObjectFixupTool::ResetDiff()
 
 TObjectPtr<UObject> SInstanceDataObjectFixupTool::CreateInstanceDataObjectSnapshot()
 {
+	using namespace UE::DataRecoveryTool::Utils;
+	using namespace UE::DataRecoveryTool::Utils::Snapshot;
+
+
 	const TSharedPtr<FTopLevelAssetPath> ClassPath = SelectedClassPath.Pin();
 
 	if (!ClassPath)
@@ -114,18 +118,14 @@ TObjectPtr<UObject> SInstanceDataObjectFixupTool::CreateInstanceDataObjectSnapsh
 
 	TObjectPtr<UObject> Owner;
 
-	if (UE::DataRecoveryTool::Utils::GetInstanceDataObjectFromSelectedClassPath(*ClassPath, &Owner) == nullptr)
+	if (GetInstanceDataObjectFromSelectedClassPath(*ClassPath, &Owner) == nullptr)
 	{
 		return nullptr;
 	}
 
 	// CPFUO the instance
-	constexpr EObjectFlags FlagsToRemove = RF_Public | RF_Standalone;
-	constexpr EObjectFlags FlagsToAdd = RF_Transient | RF_DuplicateTransient;
-
-	constexpr EObjectFlags Flags = FlagsToAdd & ~FlagsToRemove;
 
-	UObject* OwnerSnapshot = NewObject<UObject>(GetTransientPackage(), Owner->GetClass(), NAME_None, Flags);
+	UObject* OwnerSnapshot = NewObject<UObject>(GetTransientPackage(), Owner->GetClass(), NAME_None, FlagsToAdd & ~FlagsToRemove);
 
 	TMap<UObject*, UObject*> ReplacementMap;
 	ReplacementMap.Add(Owner, OwnerSnapshot);
@@ -233,26 +233,40 @@ TSharedRef<SWidget> SInstanceDataObjectFixupTool::MakeDiffViewPanel()
 
 void SInstanceDataObjectFixupTool::GenerateDetailsViews()
 {
-	TObjectPtr<UObject> InstanceDataObjectSnapshot = CreateInstanceDataObjectSnapshot();
 
-	if (InstanceDataObjectSnapshot == nullptr)
+#if USE_SNAPSHOT_INSTANCE
+	UObject* InstanceDataObject = CreateInstanceDataObjectSnapshot();
+
+	if (InstanceDataObject == nullptr)
 	{
 		ResetDiff();
 		return;
 	}
+#else
+	const TSharedPtr<FTopLevelAssetPath> ClassPath = SelectedClassPath.Pin();
 
-	// apply staged transforms if any
-	if (TSharedPtr Transforms = StagedTransforms.Pin(); const UE::FInstanceDataTransformSet* TransformSet = Transforms->Find(*SelectedClassPath.Pin()))
+	if (!ClassPath)
 	{
-		const UObject* Owner = UE::FPropertyBagRepository::Get().FindInstanceForDataObject(InstanceDataObjectSnapshot);
-		InstanceDataObjectSnapshot = UE::FInstanceDataTransforms::Get().ApplyTransformSet(*TransformSet, InstanceDataObjectSnapshot, const_cast<UObject*>(Owner));
+		ResetDiff();
+		return;
 	}
 
+	UObject* InstanceDataObject = UE::DataRecoveryTool::Utils::GetInstanceDataObjectFromSelectedClassPath(*ClassPath);
+	if (!InstanceDataObject)
+	{
+		ResetDiff();
+		return;
+	}
+#endif
+
+	UObject* Owner = const_cast<UObject*>(UE::FPropertyBagRepository::Get().FindInstanceForDataObject(InstanceDataObject));
+	UE::DataRecoveryTool::Utils::ApplyTransforms(StagedTransforms, InstanceDataObject, Owner);
+
 	constexpr int32 LeftIndex = 0;
 	constexpr int32 RightIndex = 1;
 
-	Panels[LeftIndex] = MakeShared<FInstanceDataObjectFixupPanel>(InstanceDataObjectSnapshot, FInstanceDataObjectFixupPanel::EViewFlags::DefaultLeftPanel, StagedTransforms);
-	Panels[RightIndex] = MakeShared<FInstanceDataObjectFixupPanel>(InstanceDataObjectSnapshot, FInstanceDataObjectFixupPanel::EViewFlags::DefaultRightPanel, StagedTransforms);
+	Panels[LeftIndex] = MakeShared<FInstanceDataObjectFixupPanel>(InstanceDataObject, FInstanceDataObjectFixupPanel::EViewFlags::DefaultLeftPanel, StagedTransforms);
+	Panels[RightIndex] = MakeShared<FInstanceDataObjectFixupPanel>(InstanceDataObject, FInstanceDataObjectFixupPanel::EViewFlags::DefaultRightPanel, StagedTransforms);
 
 	// generate the panels
 	Panels[LeftIndex]->GenerateDetailsView(true);
@@ -347,20 +361,32 @@ FReply SInstanceDataObjectFixupTool::OnResetInstanceDataObjectSnapshot()
 		return FReply::Unhandled();
 	}
 
-	UClass* LoadedClass = Cast<UClass>(StaticLoadObject(UClass::StaticClass(), nullptr, ClassPath->ToString()));
-	if (LoadedClass == nullptr)
-	{
-		return FReply::Unhandled();
-	}
-
-	UE::FInstanceDataTransformSet* TransformSet = StagedTransforms.Pin()->Find(*ClassPath);
-	if (TransformSet == nullptr)
+	if (const TSharedPtr Transforms = StagedTransforms.Pin())
 	{
-		return FReply::Unhandled();
+		for (const TPair<FTopLevelAssetPath, UE::FInstanceDataTransformSet>& TransformPair : *Transforms)
+		{
+			if (FTopLevelAssetPath TransformPath = TransformPair.Key;
+				UE::DataRecoveryTool::Utils::ContainsPath(*ClassPath, TransformPath))
+			{
+				if (UE::FInstanceDataTransformSet* TransformSet = Transforms->Find(TransformPath))
+				{
+					if (UStruct* LoadedStruct = Cast<UStruct>(StaticLoadObject(UStruct::StaticClass(), nullptr, TransformPath.ToString())))
+					{
+						*TransformSet = UE::FInstanceDataTransforms::Get().GetTransformSet(*LoadedStruct);
+					}
+					else
+					{
+						return FReply::Unhandled();
+					}
+				}
+				else
+				{
+					return FReply::Unhandled();
+				}
+			}
+		}
 	}
 
-	*TransformSet = UE::FInstanceDataTransforms::Get().GetTransformSet(*LoadedClass);
-
 	GenerateDetailsViews();
 
 	return FReply::Handled();
@@ -369,16 +395,36 @@ FReply SInstanceDataObjectFixupTool::OnResetInstanceDataObjectSnapshot()
 
 bool SInstanceDataObjectFixupTool::CanResetInstanceDataObjectSnapshot() const
 {
-	const UE::FInstanceDataTransformSet* TransformSet = GetSelectedStagedTransforms();
+	TSharedPtr<FTopLevelAssetPath> ClassPath = SelectedClassPath.Pin();
+	if (ClassPath == nullptr)
+	{
+		return false;
+	}
+
+	if (const TSharedPtr Transforms = StagedTransforms.Pin())
+	{
+		for (const TPair<FTopLevelAssetPath, UE::FInstanceDataTransformSet>& TransformPair : *Transforms)
+		{
+			if (FTopLevelAssetPath TransformPath = TransformPair.Key;
+				UE::DataRecoveryTool::Utils::ContainsPath(*ClassPath, TransformPath))
+			{
+				UE::FInstanceDataTransformSet* TransformSet = Transforms->Find(TransformPath);
+				if (TransformSet && TransformSet->Operations.Num() > 0)
+				{
+					return true;
+				}
+			}
+		}
+	}
 
-	return TransformSet && !TransformSet->Operations.IsEmpty();
+	return false;
 }
 
 UE::FInstanceDataTransformSet* SInstanceDataObjectFixupTool::GetSelectedStagedTransforms() const
 {
 	if (const TSharedPtr<FTopLevelAssetPath> SelectedPath = SelectedClassPath.Pin())
 	{
-		if (const TSharedPtr<TMap<FTopLevelAssetPath, UE::FInstanceDataTransformSet>> Transforms = StagedTransforms.Pin())
+		if (const TSharedPtr<TMap<FTopLevelAssetPath, UE::FInstanceDataTransformSet>>& Transforms = StagedTransforms.Pin())
 		{
 			if (UE::FInstanceDataTransformSet* Found = Transforms->Find(*SelectedPath))
 			{
@@ -394,3 +440,4 @@ END_SLATE_FUNCTION_BUILD_OPTIMIZATION
 
 #undef LOCTEXT_NAMESPACE
 
+#undef USE_SNAPSHOT_INSTANCE
\ No newline at end of file
