Allow unconditional type name impersonation

Type name impersonation has been gated on bImpersonateProperties primarily to avoid performance impact for systems that do not require it. Provide type name impersonation unconditionally to avoid the often-subtle issues that occur when using the wrong type name.

Removes scoped enabling of bImpersonateProperties anywhere that seems to be doing it only for SaveTypeName and functions that depend on it like FPropertyTypeName(const FProperty*) and FPropertyPath::ToPropertyPathName().

diff --git a/Engine/Plugins/Experimental/InstanceDataObjectFixupTool/Source/Private/InstanceDataObjectFixupPanel.cpp b/Engine/Plugins/Experimental/InstanceDataObjectFixupTool/Source/Private/InstanceDataObjectFixupPanel.cpp
--- a/Engine/Plugins/Experimental/InstanceDataObjectFixupTool/Source/Private/InstanceDataObjectFixupPanel.cpp
+++ b/Engine/Plugins/Experimental/InstanceDataObjectFixupTool/Source/Private/InstanceDataObjectFixupPanel.cpp
@@ -343,8 +343,6 @@ UE::FInstanceDataTransformSet* FInstanceDataObjectFixupPanel::GetStagedTransform
 
 void FInstanceDataObjectFixupPanel::RedirectProperty(const FPropertyPath& From, const FPropertyPath& To)
 {
-	FUObjectSerializeContext* SerializeContext = FUObjectThreadContext::Get().GetSerializeContext();
-	TGuardValue<bool> ScopedImpersonateProperties(SerializeContext->bImpersonateProperties, true);
 	TArray<UE::FPropertyPathName> FromPathNames = From.ToPropertyPathName();
 	TArray<UE::FPropertyPathName> ToPathNames = To.ToPropertyPathName();
 
@@ -406,8 +404,6 @@ bool FInstanceDataObjectFixupPanel::PropertyNeedsRedirect(const FPropertyPath& P
 		return false;
 	}
 
-	FUObjectSerializeContext* SerializeContext = FUObjectThreadContext::Get().GetSerializeContext();
-	TGuardValue<bool> ScopedImpersonateProperties(SerializeContext->bImpersonateProperties, true);
 	TArray<UE::FPropertyPathName> PathNames = Path.ToPropertyPathName();
 	if (PathNames.Num() != 1)
 	{
@@ -420,8 +416,6 @@ bool FInstanceDataObjectFixupPanel::PropertyNeedsRedirect(const FPropertyPath& P
 
 void FInstanceDataObjectFixupPanel::MarkForDelete(const FPropertyPath& CurrentPath)
 {
-	FUObjectSerializeContext* SerializeContext = FUObjectThreadContext::Get().GetSerializeContext();
-	TGuardValue<bool> ScopedImpersonateProperties(SerializeContext->bImpersonateProperties, true);
 	TArray<UE::FPropertyPathName> srcPathNames = CurrentPath.ToPropertyPathName();
 
 	if (srcPathNames.Num() != 1)
diff --git a/Engine/Source/Runtime/CoreUObject/Private/UObject/EnumProperty.cpp b/Engine/Source/Runtime/CoreUObject/Private/UObject/EnumProperty.cpp
--- a/Engine/Source/Runtime/CoreUObject/Private/UObject/EnumProperty.cpp
+++ b/Engine/Source/Runtime/CoreUObject/Private/UObject/EnumProperty.cpp
@@ -302,7 +302,6 @@ void FEnumProperty::SerializeItem(FStructuredArchive::FSlot Slot, void* Value, v
 			FUObjectSerializeContext* SerializeContext = FUObjectThreadContext::Get().GetSerializeContext();
 			if (UNLIKELY(SerializeContext->bTrackUnknownEnumNames))
 			{
-				TGuardValue<bool> ScopedImpersonateProperties(SerializeContext->bImpersonateProperties, true);
 				UE::FUnknownEnumNames(SerializeContext->SerializedObject).Add(Enum, UE::FindOriginalType(this), EnumValueName);
 			}
 		#endif
@@ -332,13 +331,9 @@ void FEnumProperty::SerializeItem(FStructuredArchive::FSlot Slot, void* Value, v
 
 				#if WITH_EDITORONLY_DATA
 					// Fix up the type name when this property is impersonating another enum type.
-					FUObjectSerializeContext* SerializeContext = FUObjectThreadContext::Get().GetSerializeContext();
-					if (SerializeContext->bImpersonateProperties)
+					if (UE::FPropertyTypeName OriginalType = UE::FindOriginalType(this); !OriginalType.IsEmpty())
 					{
-						if (UE::FPropertyTypeName OriginalType = UE::FindOriginalType(this); !OriginalType.IsEmpty())
-						{
-							EnumValueName = FName(EnumValueName.ToString().Replace(*Enum->GetName(), *OriginalType.GetName().ToString()));
-						}
+						EnumValueName = FName(EnumValueName.ToString().Replace(*Enum->GetName(), *OriginalType.GetName().ToString()));
 					}
 				#endif
 				}
diff --git a/Engine/Source/Runtime/CoreUObject/Private/UObject/InstanceDataObjectUtils.cpp b/Engine/Source/Runtime/CoreUObject/Private/UObject/InstanceDataObjectUtils.cpp
--- a/Engine/Source/Runtime/CoreUObject/Private/UObject/InstanceDataObjectUtils.cpp
+++ b/Engine/Source/Runtime/CoreUObject/Private/UObject/InstanceDataObjectUtils.cpp
@@ -884,11 +884,7 @@ namespace UE
 		bool bHasFlags = false;
 
 		// Use the original type name because the template may be a fallback enum or an IDO.
-		FPropertyTypeName EnumTypeName;
-		{
-			TGuardValue<bool> ImpersonateScope(FUObjectThreadContext::Get().GetSerializeContext()->bImpersonateProperties, true);
-			EnumTypeName = FindOriginalType(Property);
-		}
+		FPropertyTypeName EnumTypeName = FindOriginalType(Property);
 		if (EnumTypeName.IsEmpty())
 		{
 			FPropertyTypeNameBuilder Builder;
@@ -1039,8 +1035,7 @@ namespace UE
 			{
 				//@note: Transfer existing metadata over as we build the InstanceDataObject from the struct or it owners, if any, this is useful for testing purposes
 				TStringBuilder<256> OriginalName;
-				if (TGuardValue<bool> ImpersonateScope(FUObjectThreadContext::Get().GetSerializeContext()->bImpersonateProperties, true);
-					const FString* OriginalType = FindOriginalTypeName(AsStructProperty))
+				if (const FString* OriginalType = FindOriginalTypeName(AsStructProperty))
 				{
 					OriginalName << *OriginalType;
 				}
diff --git a/Engine/Source/Runtime/CoreUObject/Private/UObject/InstanceDataTransforms.cpp b/Engine/Source/Runtime/CoreUObject/Private/UObject/InstanceDataTransforms.cpp
--- a/Engine/Source/Runtime/CoreUObject/Private/UObject/InstanceDataTransforms.cpp
+++ b/Engine/Source/Runtime/CoreUObject/Private/UObject/InstanceDataTransforms.cpp
@@ -304,9 +304,6 @@ namespace UE::Private::IDT
 	 */
 	ETransformResult RemoveProperty(TObjectPtr<UObject>& IDO, TObjectPtr<UObject>& Owner, const FPropertyPathName& BasePath, const FTransformParams& Params)
 	{
-		FUObjectSerializeContext* SerializeContext = FUObjectThreadContext::Get().GetSerializeContext();
-		TGuardValue<bool> ScopedImpersonateProperties(SerializeContext->bImpersonateProperties, true);
-
 		TSharedPtr<FPropertyPathNameTree> UnknownPropertyTree = FUnknownPropertyTree(Owner).FindOrCreate();
 		FUnknownEnumNames UnknownEnumNames = FUnknownEnumNames(Owner);
 
@@ -355,9 +352,6 @@ namespace UE::Private::IDT
 	 */
 	ETransformResult RedirectProperty(TObjectPtr<UObject>& IDO, TObjectPtr<UObject>& Owner, const FPropertyPathName& BasePath, const FTransformParams& Params)
 	{
-		FUObjectSerializeContext* SerializeContext = FUObjectThreadContext::Get().GetSerializeContext();
-		TGuardValue<bool> ScopedImpersonateProperties(SerializeContext->bImpersonateProperties, true);
-
 		TSharedPtr<FPropertyPathNameTree> UnknownPropertyTree = FUnknownPropertyTree(Owner).FindOrCreate();
 		FUnknownEnumNames UnknownEnumNames = FUnknownEnumNames(Owner);
 
diff --git a/Engine/Source/Runtime/CoreUObject/Private/UObject/OverriddenPropertySet.cpp b/Engine/Source/Runtime/CoreUObject/Private/UObject/OverriddenPropertySet.cpp
--- a/Engine/Source/Runtime/CoreUObject/Private/UObject/OverriddenPropertySet.cpp
+++ b/Engine/Source/Runtime/CoreUObject/Private/UObject/OverriddenPropertySet.cpp
@@ -324,17 +324,11 @@ FString FOverriddenPropertyPath::ToString() const
 //----------------------------------------------------------------------//
 FOverriddenPropertyNodeID::FOverriddenPropertyNodeID(TNotNull<const FProperty*> Property)
 {
+	Path = Property->GetFName();
+
 	// append typename to the end of the property ID
 	UE::FPropertyTypeNameBuilder TypeNameBuilder;
-#if WITH_EDITORONLY_DATA
-	{
-		// use property impersonation for SaveTypeName so that keys don't change when classes die
-		FUObjectSerializeContext* SerializeContext = FUObjectThreadContext::Get().GetSerializeContext();
-        TGuardValue<bool> ScopedImpersonateProperties(SerializeContext->bImpersonateProperties, true);
-        Property->SaveTypeName(TypeNameBuilder);
-	}
-#endif
-	Path = Property->GetFName();
+	Property->SaveTypeName(TypeNameBuilder);
 	for (const UE::FPropertyTypeNameNode& Node : TypeNameBuilder)
 	{
 		Path += Node.Name;
diff --git a/Engine/Source/Runtime/CoreUObject/Private/UObject/PropertyBagRepository.cpp b/Engine/Source/Runtime/CoreUObject/Private/UObject/PropertyBagRepository.cpp
--- a/Engine/Source/Runtime/CoreUObject/Private/UObject/PropertyBagRepository.cpp
+++ b/Engine/Source/Runtime/CoreUObject/Private/UObject/PropertyBagRepository.cpp
@@ -234,14 +234,9 @@ void FPropertyBagRepository::ReassociateObjects(const TMap<UObject*, UObject*>&
 
 static FProperty* FindPropertyByNameAndType(const UStruct* Struct, const FProperty* SourceProperty)
 {
-	// use property impersonation for SaveTypeName so that keys of IDOs and non-IDOs match
-	FUObjectSerializeContext* SerializeContext = FUObjectThreadContext::Get().GetSerializeContext();
-	TGuardValue<bool> ScopedImpersonateProperties(SerializeContext->bImpersonateProperties, true);
-
 	return FindPropertyByNameAndTypeName(Struct, SourceProperty->GetFName(), FPropertyTypeName(SourceProperty));
 }
 
-	
 static void CopyProperty(const FProperty* SourceProperty, const void* SourceValue, const FProperty* DestProperty, void* DestValue)
 {
 	check(SourceProperty->GetID() == DestProperty->GetID());
@@ -734,9 +729,6 @@ bool FPropertyBagRepository::RequiresFixup(const UObject* Object, const FPropert
 		return false;
 	}
 
-	FUObjectSerializeContext* SerializeContext = FUObjectThreadContext::Get().GetSerializeContext();
-	TGuardValue<bool> ScopedImpersonateProperties(SerializeContext->bImpersonateProperties, true);
-
 	// look for either the full path in the tree, or a leaf node in the tree that matches the beginning of the path.
 	// the latter scenario can occur because if a struct is a leaf, all it's child nodes are implicitly also unknown properties so they're not included in the tree
 	UE::FPropertyPathName SubPath = Path;
diff --git a/Engine/Source/Runtime/CoreUObject/Private/UObject/PropertyByte.cpp b/Engine/Source/Runtime/CoreUObject/Private/UObject/PropertyByte.cpp
--- a/Engine/Source/Runtime/CoreUObject/Private/UObject/PropertyByte.cpp
+++ b/Engine/Source/Runtime/CoreUObject/Private/UObject/PropertyByte.cpp
@@ -77,7 +77,6 @@ void FByteProperty::SerializeItem( FStructuredArchive::FSlot Slot, void* Value,
 			FUObjectSerializeContext* SerializeContext = FUObjectThreadContext::Get().GetSerializeContext();
 			if (UNLIKELY(SerializeContext->bTrackUnknownEnumNames))
 			{
-				TGuardValue<bool> ScopedImpersonateProperties(SerializeContext->bImpersonateProperties, true);
 				UE::FUnknownEnumNames(SerializeContext->SerializedObject).Add(Enum, UE::FindOriginalType(this), EnumValueName);
 			}
 		#endif
@@ -98,13 +97,9 @@ void FByteProperty::SerializeItem( FStructuredArchive::FSlot Slot, void* Value,
 
 		#if WITH_EDITORONLY_DATA
 			// Fix up the type name when this property is impersonating another enum type.
-			FUObjectSerializeContext* SerializeContext = FUObjectThreadContext::Get().GetSerializeContext();
-			if (SerializeContext->bImpersonateProperties)
+			if (UE::FPropertyTypeName OriginalType = UE::FindOriginalType(this); !OriginalType.IsEmpty())
 			{
-				if (UE::FPropertyTypeName OriginalType = UE::FindOriginalType(this); !OriginalType.IsEmpty())
-				{
-					EnumValueName = FName(EnumValueName.ToString().Replace(*Enum->GetName(), *OriginalType.GetName().ToString()));
-				}
+				EnumValueName = FName(EnumValueName.ToString().Replace(*Enum->GetName(), *OriginalType.GetName().ToString()));
 			}
 		#endif
 		}
diff --git a/Engine/Source/Runtime/CoreUObject/Private/UObject/PropertyHelper.cpp b/Engine/Source/Runtime/CoreUObject/Private/UObject/PropertyHelper.cpp
--- a/Engine/Source/Runtime/CoreUObject/Private/UObject/PropertyHelper.cpp
+++ b/Engine/Source/Runtime/CoreUObject/Private/UObject/PropertyHelper.cpp
@@ -327,20 +327,11 @@ UObject* FindObjectByTypePath(UClass* Class, FPropertyTypeName TypePath)
 #if WITH_EDITORONLY_DATA
 const FString* FindOriginalTypeName(const UField* Field)
 {
-	if (Field && FUObjectThreadContext::Get().GetSerializeContext()->bImpersonateProperties)
-	{
-		return Field->FindMetaData(NAME_OriginalType);
-	}
-	return nullptr;
+	return Field ? Field->FindMetaData(NAME_OriginalType) : nullptr;
 }
 
 const FString* FindOriginalTypeName(const FProperty* Property)
 {
-	if (!FUObjectThreadContext::Get().GetSerializeContext()->bImpersonateProperties)
-	{
-		return nullptr;
-	}
-
 	// Prioritize metadata on the property over metadata on the type.
 	if (const FString* OriginalType = Property->FindMetaData(NAME_OriginalType))
 	{
