[High Priority] [Crash Fix] Localization build crash fix.

No longer using CVar ref in plugin causing a circular dependency

diff --git a/Engine/Plugins/Experimental/InstanceDataObjectFixupTool/Source/Private/DataRecoveryTool.cpp b/Engine/Plugins/Experimental/InstanceDataObjectFixupTool/Source/Private/DataRecoveryTool.cpp
--- a/Engine/Plugins/Experimental/InstanceDataObjectFixupTool/Source/Private/DataRecoveryTool.cpp
+++ b/Engine/Plugins/Experimental/InstanceDataObjectFixupTool/Source/Private/DataRecoveryTool.cpp
@@ -30,7 +30,7 @@
 #include "Widgets/STedsHierarchyViewer.h"
 #include "Serialization/ObjectWriter.h"
 #include "Misc/ScopedSlowTask.h"
-
+#include "UObject/PropertyBagRepository.h"
 
 #define USE_SAVE_BUTTON 0
 #if USE_SAVE_BUTTON
diff --git a/Engine/Plugins/Experimental/InstanceDataObjectFixupTool/Source/Private/DataRecoveryToolUtils.cpp b/Engine/Plugins/Experimental/InstanceDataObjectFixupTool/Source/Private/DataRecoveryToolUtils.cpp
--- a/Engine/Plugins/Experimental/InstanceDataObjectFixupTool/Source/Private/DataRecoveryToolUtils.cpp
+++ b/Engine/Plugins/Experimental/InstanceDataObjectFixupTool/Source/Private/DataRecoveryToolUtils.cpp
@@ -1,6 +1,7 @@
 // Copyright Epic Games, Inc. All Rights Reserved.
 
 #include "DataRecoveryToolUtils.h"
+#include "UObject/PropertyBagRepository.h"
 
 #define LOCTEXT_NAMESPACE "DataRecoveryToolModule"
 
diff --git a/Engine/Plugins/Experimental/InstanceDataObjectFixupTool/Source/Private/DataRecoveryToolUtils.h b/Engine/Plugins/Experimental/InstanceDataObjectFixupTool/Source/Private/DataRecoveryToolUtils.h
--- a/Engine/Plugins/Experimental/InstanceDataObjectFixupTool/Source/Private/DataRecoveryToolUtils.h
+++ b/Engine/Plugins/Experimental/InstanceDataObjectFixupTool/Source/Private/DataRecoveryToolUtils.h
@@ -3,10 +3,7 @@
 #pragma once
 
 #include "CoreMinimal.h"
-#include "PropertyPath.h"
-#include "UObject/PropertyStateTracking.h"
-#include "UObject/PropertyPathName.h"
-#include "UObject/PropertyBagRepository.h"
+#include "Elements/Interfaces/TypedElementDataStorageInterface.h"
 #include "UObject/InstanceDataTransforms.h"
 #include "TedsAlertColumns.h"
 #include "TedsAlerts.h"
diff --git a/Engine/Plugins/Experimental/InstanceDataObjectFixupTool/Source/Private/InstanceDataObjectFixupToolModule.cpp b/Engine/Plugins/Experimental/InstanceDataObjectFixupTool/Source/Private/InstanceDataObjectFixupToolModule.cpp
--- a/Engine/Plugins/Experimental/InstanceDataObjectFixupTool/Source/Private/InstanceDataObjectFixupToolModule.cpp
+++ b/Engine/Plugins/Experimental/InstanceDataObjectFixupTool/Source/Private/InstanceDataObjectFixupToolModule.cpp
@@ -13,32 +13,12 @@
 
 #define LOCTEXT_NAMESPACE "InstanceDataObjectFixupToolModule"
 
-static const TCHAR* CVarEnableDataRecoveryToolName = TEXT("IDO.EnableDataRecoveryTool");
 
 void FInstanceDataObjectFixupToolModule::StartupModule()
 {
-	if (!FModuleManager::Get().IsModuleLoaded("EntityEditor"))
-	{
-		FModuleManager::Get().LoadModuleChecked("EntityEditor");
-	}
-
-	CVarEnableDataRecoveryTool = IConsoleManager::Get().FindConsoleVariable(CVarEnableDataRecoveryToolName);
-	if (!CVarEnableDataRecoveryTool)
-	{
-		return;
-	}
 
-	CVarChangedHandle =
-		CVarEnableDataRecoveryTool->AsVariable()->OnChangedDelegate()
-        .AddRaw(this, &FInstanceDataObjectFixupToolModule::OnCVarChanged);
-
-	OnCVarChanged(CVarEnableDataRecoveryTool);
 }
 
-void FInstanceDataObjectFixupToolModule::OnCVarChanged(IConsoleVariable* Var)
-{
-	Var->GetBool() == true ? RegisterTabSpawners() : UnregisterTabSpawners();
-}
 
 void FInstanceDataObjectFixupToolModule::RegisterTabSpawners(const TOptional<FTopLevelAssetPath>& ClassPath) const
 {
@@ -69,12 +49,6 @@ void FInstanceDataObjectFixupToolModule::UnregisterTabSpawners() const
 void FInstanceDataObjectFixupToolModule::ShutdownModule()
 {
 	UnregisterTabSpawners();
-	if (CVarEnableDataRecoveryTool)
-	{
-		CVarEnableDataRecoveryTool->AsVariable()->OnChangedDelegate().Remove(CVarChangedHandle);
-		CVarEnableDataRecoveryTool = nullptr;
-	}
-	CVarChangedHandle.Reset();
 }
 
 TSharedRef<SDockTab> FInstanceDataObjectFixupToolModule::CreateDataRecoveryToolTab(const TOptional<FTopLevelAssetPath>& ClassPath) const
diff --git a/Engine/Plugins/Experimental/InstanceDataObjectFixupTool/Source/Public/InstanceDataObjectFixupToolModule.h b/Engine/Plugins/Experimental/InstanceDataObjectFixupTool/Source/Public/InstanceDataObjectFixupToolModule.h
--- a/Engine/Plugins/Experimental/InstanceDataObjectFixupTool/Source/Public/InstanceDataObjectFixupToolModule.h
+++ b/Engine/Plugins/Experimental/InstanceDataObjectFixupTool/Source/Public/InstanceDataObjectFixupToolModule.h
@@ -21,13 +21,8 @@ public:
 
 	UE_API void CreateDataRecoveryToolDialog(const TOptional<FTopLevelAssetPath>& ClassPath) const;
 
-private:
-	void OnCVarChanged(IConsoleVariable* Var);
-	void RegisterTabSpawners(const TOptional<FTopLevelAssetPath>& ClassPath = {}) const;
-	void UnregisterTabSpawners() const;
-
-	IConsoleVariable* CVarEnableDataRecoveryTool;
-	FDelegateHandle CVarChangedHandle;
+	UE_API void RegisterTabSpawners(const TOptional<FTopLevelAssetPath>& ClassPath = {}) const;
+	UE_API void UnregisterTabSpawners() const;
 };
 
 #undef UE_API
