[Improvement] DRT: Add CVar to control IDT serialization

diff --git a/Engine/Plugins/Experimental/InstanceDataObjectFixupTool/Source/Private/DataRecoveryTool.cpp b/Engine/Plugins/Experimental/InstanceDataObjectFixupTool/Source/Private/DataRecoveryTool.cpp
--- a/Engine/Plugins/Experimental/InstanceDataObjectFixupTool/Source/Private/DataRecoveryTool.cpp
+++ b/Engine/Plugins/Experimental/InstanceDataObjectFixupTool/Source/Private/DataRecoveryTool.cpp
@@ -32,10 +32,7 @@
 #include "Misc/ScopedSlowTask.h"
 #include "UObject/PropertyBagRepository.h"
 
-#define USE_SAVE_BUTTON 0
-#if USE_SAVE_BUTTON
 #include "Widgets/Layout/SSeparator.h"
-#endif // USE_SAVE_BUTTON
 
 #define USE_OPEN_EXTERNAL_BUTTON 0
 
@@ -43,10 +40,9 @@
 
 BEGIN_SLATE_FUNCTION_BUILD_OPTIMIZATION
 
-
-
 using namespace UE::DataRecoveryTool::StyleSheet;
 
+
 void SDataRecoveryTool::Construct(const FArguments& InArgs)
 {
 	SelectedClassPath =  MakeShared<FTopLevelAssetPath>(InArgs._SelectedClassPath.IsSet() ? *InArgs._SelectedClassPath : nullptr);
@@ -79,18 +75,13 @@ void SDataRecoveryTool::Construct(const FArguments& InArgs)
 		];
 }
 
-
-
 TSharedRef<SWidget> SDataRecoveryTool::MakeTopBar()
 {
-	return
-		SNew(SBorder)
-		.BorderImage(FAppStyle::Get().GetBrush("Brushes.Panel"))
-		.Padding(Padding::None)
-		[
-			SNew(SHorizontalBox)
-#if USE_SAVE_BUTTON
-			+SHorizontalBox::Slot()
+	TSharedRef<SHorizontalBox> TopBarContent = SNew(SHorizontalBox);
+
+	if (UE::FInstanceDataTransforms::Get().IsSerializationEnabled())
+	{
+		TopBarContent->AddSlot()
 			.AutoWidth()
 			.Padding(0)
 			[
@@ -105,8 +96,9 @@ TSharedRef<SWidget> SDataRecoveryTool::MakeTopBar()
 					.Image(FAppStyle::GetBrush("Icons.Save"))
 					.ToolTipText(INVTEXT("Save"))
 				]
-			]
-			+ SHorizontalBox::Slot()
+			];
+
+		TopBarContent->AddSlot()
 			.AutoWidth()
 			.VAlign(VAlign_Fill)
 			[
@@ -114,72 +106,80 @@ TSharedRef<SWidget> SDataRecoveryTool::MakeTopBar()
 				.Thickness(Padding::Small * .5f)
 				.Orientation(Orient_Vertical)
 				.SeparatorImage(FAppStyle::Get().GetBrush("Brushes.Background"))
-			]
-#endif // USE_SAVE_BUTTON
+			];
+	}
+
+	TopBarContent->AddSlot()
+		.FillWidth(1.f)
+		[
+			SNew(SHorizontalBox)
 			+ SHorizontalBox::Slot()
-			.FillWidth(1.f)
+			.AutoWidth()
+			.HAlign(HAlign_Left)
 			[
-				SNew(SHorizontalBox)
-				+ SHorizontalBox::Slot()
-				.AutoWidth()
-				.HAlign(HAlign_Left)
+				SNew(SButton)
+				.HAlign(HAlign_Center)
+				.VAlign(VAlign_Center)
+				.ContentPadding(Padding::Normal)
+				.ButtonStyle(FAppStyle::Get(), "SimpleButton")
+				.OnClicked(this, &SDataRecoveryTool::OnApplyClicked)
 				[
-					SNew(SButton)
-					.HAlign(HAlign_Center)
-					.VAlign(VAlign_Center)
-					.ContentPadding(Padding::Normal)
-					.ButtonStyle(FAppStyle::Get(), "SimpleButton")
-					.OnClicked(this, &SDataRecoveryTool::OnApplyClicked)
+					SNew(SHorizontalBox)
+					+ SHorizontalBox::Slot()
+					.AutoWidth()
 					[
-						SNew(SHorizontalBox)
-						+ SHorizontalBox::Slot()
-						.AutoWidth()
-						[
-							SNew(SImage)
-							.Image(FAppStyle::GetBrush("LevelEditor.Recompile.Small"))
-						]
-						+ SHorizontalBox::Slot()
-						.Padding(Padding::Small, Padding::None, Padding::None, Padding::None)
-						.AutoWidth()
-						[
-							SNew(STextBlock)
-							.Text(LOCTEXT("ApplyButtonText", "Apply"))
-						]
+						SNew(SImage)
+						.Image(FAppStyle::GetBrush("LevelEditor.Recompile.Small"))
+					]
+					+ SHorizontalBox::Slot()
+					.Padding(Padding::Small, Padding::None, Padding::None, Padding::None)
+					.AutoWidth()
+					[
+						SNew(STextBlock)
+						.Text(LOCTEXT("ApplyButtonText", "Apply"))
 					]
 				]
-				+ SHorizontalBox::Slot()
-				.FillWidth(1.f)
-				[
-					SNullWidget::NullWidget
-				]
+			]
+		+ SHorizontalBox::Slot()
+			.FillWidth(1.f)
+			[
+				SNullWidget::NullWidget
+			]
 #if USE_OPEN_EXTERNAL_BUTTON
-				+ SHorizontalBox::Slot()
-				.HAlign(HAlign_Right)
+			+ SHorizontalBox::Slot()
+			.HAlign(HAlign_Right)
+			[
+				SNew(SButton)
+				.HAlign(HAlign_Center)
+				.VAlign(VAlign_Center)
+				.ContentPadding(Padding::Normal)
+				.ButtonStyle(FAppStyle::Get(), "SimpleButton")
 				[
-					SNew(SButton)
-					.HAlign(HAlign_Center)
-					.VAlign(VAlign_Center)
-					.ContentPadding(Padding::Normal)
-					.ButtonStyle(FAppStyle::Get(), "SimpleButton")
+					SNew(SHorizontalBox)
+					+ SHorizontalBox::Slot()
+					.AutoWidth()
 					[
-						SNew(SHorizontalBox)
-						+ SHorizontalBox::Slot()
-						.AutoWidth()
-						[
-							SNew(SImage)
-							.Image(FAppStyle::GetBrush("Icons.OpenInExternalEditor"))
-						]
-						+ SHorizontalBox::Slot()
-						.Padding(Padding::Small, Padding::None, Padding::None, Padding::None)
-						.AutoWidth()
-						[
-							SNew(STextBlock)
-							.Text(LOCTEXT("EditInVSCodeButton", "Edit in VS Code"))
-						]
+						SNew(SImage)
+						.Image(FAppStyle::GetBrush("Icons.OpenInExternalEditor"))
+					]
+					+ SHorizontalBox::Slot()
+					.Padding(Padding::Small, Padding::None, Padding::None, Padding::None)
+					.AutoWidth()
+					[
+						SNew(STextBlock)
+						.Text(LOCTEXT("EditInVSCodeButton", "Edit in VS Code"))
 					]
 				]
-#endif // USE_OPEN_EXTERNAL_BUTTON
 			]
+#endif // USE_OPEN_EXTERNAL_BUTTON
+		];
+
+	return
+		SNew(SBorder)
+		.BorderImage(FAppStyle::Get().GetBrush("Brushes.Panel"))
+		.Padding(Padding::None)
+		[
+			TopBarContent
 		];
 }
 
diff --git a/Engine/Source/Runtime/CoreUObject/Internal/UObject/InstanceDataTransforms.h b/Engine/Source/Runtime/CoreUObject/Internal/UObject/InstanceDataTransforms.h
--- a/Engine/Source/Runtime/CoreUObject/Internal/UObject/InstanceDataTransforms.h
+++ b/Engine/Source/Runtime/CoreUObject/Internal/UObject/InstanceDataTransforms.h
@@ -63,6 +63,8 @@ namespace UE
 		COREUOBJECT_API TObjectPtr<UObject> ApplyTransformSet(const FInstanceDataTransformSet& TransformSet, TNotNull<UObject*> IDO, TNotNull<UObject*> Owner);
 		COREUOBJECT_API void SaveTransformSet(const FInstanceDataTransformSet& TransformSet);
 
+		COREUOBJECT_API bool IsSerializationEnabled();
+
 	private:
 
 		const FInstanceDataTransformSet* FindTransformsInternal(const FTopLevelAssetPath& StructPath);
diff --git a/Engine/Source/Runtime/CoreUObject/Private/UObject/InstanceDataTransforms.cpp b/Engine/Source/Runtime/CoreUObject/Private/UObject/InstanceDataTransforms.cpp
--- a/Engine/Source/Runtime/CoreUObject/Private/UObject/InstanceDataTransforms.cpp
+++ b/Engine/Source/Runtime/CoreUObject/Private/UObject/InstanceDataTransforms.cpp
@@ -27,7 +27,7 @@
 #include "Misc/Paths.h"
 #include "Misc/PackageName.h"
 #include "VerseVM/VVMNames.h"
-
+#include "HAL/IConsoleManager.h"
 
 DEFINE_LOG_CATEGORY_STATIC(LogInstanceDataTransforms, Display, Display);
 
@@ -35,6 +35,11 @@ DEFINE_LOG_CATEGORY_STATIC(LogInstanceDataTransforms, Display, Display);
 #define REGISTER_IDT_OPERATION(FunctionName) \
 UE::Private::FRegesterInstanceDataTransformHelper _Registered##FunctionName{TEXT(#FunctionName), FunctionName}; \
 
+static bool bEnableInstanceDataTransformSerialization = false;
+static FAutoConsoleVariableRef CVarEnableInstanceDataTransformSerialization(
+	TEXT("IDO.EnableInstanceDataTransformSerialization"),
+	bEnableInstanceDataTransformSerialization,
+	TEXT("Enable serialization of Instance Data Transforms"));
 
 namespace UE::Private
 {
@@ -471,9 +476,12 @@ namespace UE
 	FInstanceDataTransformSet FInstanceDataTransforms::GetTransformSet(UStruct& ClassOrStruct)
 	{
 		const FTopLevelAssetPath StructPath = ClassOrStruct.GetStructPathName();
-		if (const FInstanceDataTransformSet* TransformSet = FindTransformsInternal(StructPath))
+		if (bEnableInstanceDataTransformSerialization)
 		{
-			return *TransformSet;
+			if (const FInstanceDataTransformSet* TransformSet = FindTransformsInternal(StructPath))
+			{
+				return *TransformSet;
+			}
 		}
 		return FInstanceDataTransformSet{.StructPath= StructPath };
 	}
@@ -507,6 +515,12 @@ namespace UE
 
 	void FInstanceDataTransforms::SaveTransformSet(const FInstanceDataTransformSet& TransformSet)
 	{
+		if (!bEnableInstanceDataTransformSerialization)
+		{
+			UE_LOG(LogInstanceDataTransforms, Error, TEXT("Current configuration does not allow you to save IDTs."));
+			return;
+		}
+
 		FString ConfigFilepath = Private::InstanceDataTransformDirectory(TransformSet.StructPath);
 		if (ConfigFilepath.IsEmpty())
 		{
@@ -574,6 +588,11 @@ namespace UE
 		}
 	}
 
+	bool FInstanceDataTransforms::IsSerializationEnabled()
+	{
+		return bEnableInstanceDataTransformSerialization;
+	}
+
 	const FInstanceDataTransformSet* FInstanceDataTransforms::FindTransformsInternal(const FTopLevelAssetPath& ClassPath)
 	{
 		// the transform config for a class is found in <RootPackageDirectory>/Config/InstanceDataTransforms.json
