[Refactor] Adding Stylesheets to DRT

Added basic stylesheet class to DRT as a base for the stylesheet future

diff --git a/Engine/Plugins/Experimental/InstanceDataObjectFixupTool/Source/Private/DataRecoveryTool.cpp b/Engine/Plugins/Experimental/InstanceDataObjectFixupTool/Source/Private/DataRecoveryTool.cpp
--- a/Engine/Plugins/Experimental/InstanceDataObjectFixupTool/Source/Private/DataRecoveryTool.cpp
+++ b/Engine/Plugins/Experimental/InstanceDataObjectFixupTool/Source/Private/DataRecoveryTool.cpp
@@ -2,9 +2,6 @@
 
 #include "DataRecoveryTool.h"
 
-
-
-#include "DataRecoveryToolStyleSheet.h"
 #include "DataStorage/Features.h"
 #include "DetailTreeNode.h"
 #include "InstanceDataObjectFixupTool.h"
@@ -32,6 +29,8 @@
 #include "Misc/ScopedSlowTask.h"
 #include "UObject/PropertyBagRepository.h"
 
+#include "DataRecoveryToolStyle.h"
+
 #include "Widgets/Layout/SSeparator.h"
 
 #define USE_OPEN_EXTERNAL_BUTTON 0
@@ -40,11 +39,11 @@
 
 BEGIN_SLATE_FUNCTION_BUILD_OPTIMIZATION
 
-using namespace UE::DataRecoveryTool::StyleSheet;
-
-
 void SDataRecoveryTool::Construct(const FArguments& InArgs)
 {
+	const float NoPadding = FDataRecoveryToolStyle::Get().GetFloat("DataRecoveryTool.Padding.None");
+	const float SmallPadding = FDataRecoveryToolStyle::Get().GetFloat("DataRecoveryTool.Padding.Small");
+
 	SelectedClassPath =  MakeShared<FTopLevelAssetPath>(InArgs._SelectedClassPath.IsSet() ? *InArgs._SelectedClassPath : nullptr);
 	StagedTransforms = MakeShared<TMap<FTopLevelAssetPath, UE::FInstanceDataTransformSet>>();
 	ChildSlot
@@ -57,7 +56,7 @@ void SDataRecoveryTool::Construct(const FArguments& InArgs)
 			]
 			+SVerticalBox::Slot()
 			.FillHeight(1.f)
-			.Padding(Padding::None, Padding::Small, Padding::None, Padding::None)
+			.Padding(NoPadding, SmallPadding, NoPadding, NoPadding)
 			[
 				SNew(SSplitter)
 				+SSplitter::Slot()
@@ -77,6 +76,9 @@ void SDataRecoveryTool::Construct(const FArguments& InArgs)
 
 TSharedRef<SWidget> SDataRecoveryTool::MakeTopBar()
 {
+	const float NoPadding = FDataRecoveryToolStyle::Get().GetFloat("DataRecoveryTool.Padding.None");
+	const float SmallPadding = FDataRecoveryToolStyle::Get().GetFloat("DataRecoveryTool.Padding.Small");
+	const float NormalPadding = FDataRecoveryToolStyle::Get().GetFloat("DataRecoveryTool.Padding.Normal");
 	TSharedRef<SHorizontalBox> TopBarContent = SNew(SHorizontalBox);
 
 	if (UE::FInstanceDataTransforms::Get().IsSerializationEnabled())
@@ -88,7 +90,7 @@ TSharedRef<SWidget> SDataRecoveryTool::MakeTopBar()
 				SNew(SButton)
 				.HAlign(HAlign_Center)
 				.VAlign(VAlign_Center)
-				.ContentPadding(Padding::Normal)
+				.ContentPadding(NormalPadding)
 				.ButtonStyle(FAppStyle::Get(), "SimpleButton")
 				.OnClicked(this, &SDataRecoveryTool::OnSaveAndApplyClicked)
 				[
@@ -103,7 +105,7 @@ TSharedRef<SWidget> SDataRecoveryTool::MakeTopBar()
 			.VAlign(VAlign_Fill)
 			[
 				SNew(SSeparator)
-				.Thickness(Padding::Small * .5f)
+				.Thickness(SmallPadding * .5f)
 				.Orientation(Orient_Vertical)
 				.SeparatorImage(FAppStyle::Get().GetBrush("Brushes.Background"))
 			];
@@ -120,7 +122,7 @@ TSharedRef<SWidget> SDataRecoveryTool::MakeTopBar()
 				SNew(SButton)
 				.HAlign(HAlign_Center)
 				.VAlign(VAlign_Center)
-				.ContentPadding(Padding::Normal)
+				.ContentPadding(NormalPadding)
 				.ButtonStyle(FAppStyle::Get(), "SimpleButton")
 				.OnClicked(this, &SDataRecoveryTool::OnApplyClicked)
 				[
@@ -132,7 +134,7 @@ TSharedRef<SWidget> SDataRecoveryTool::MakeTopBar()
 						.Image(FAppStyle::GetBrush("LevelEditor.Recompile.Small"))
 					]
 					+ SHorizontalBox::Slot()
-					.Padding(Padding::Small, Padding::None, Padding::None, Padding::None)
+					.Padding(SmallPadding, NoPadding, NoPadding, NoPadding)
 					.AutoWidth()
 					[
 						SNew(STextBlock)
@@ -152,7 +154,7 @@ TSharedRef<SWidget> SDataRecoveryTool::MakeTopBar()
 				SNew(SButton)
 				.HAlign(HAlign_Center)
 				.VAlign(VAlign_Center)
-				.ContentPadding(Padding::Normal)
+				.ContentPadding(NormalPadding)
 				.ButtonStyle(FAppStyle::Get(), "SimpleButton")
 				[
 					SNew(SHorizontalBox)
@@ -163,7 +165,7 @@ TSharedRef<SWidget> SDataRecoveryTool::MakeTopBar()
 						.Image(FAppStyle::GetBrush("Icons.OpenInExternalEditor"))
 					]
 					+ SHorizontalBox::Slot()
-					.Padding(Padding::Small, Padding::None, Padding::None, Padding::None)
+					.Padding(SmallPadding, NoPadding, NoPadding, NoPadding)
 					.AutoWidth()
 					[
 						SNew(STextBlock)
@@ -177,7 +179,7 @@ TSharedRef<SWidget> SDataRecoveryTool::MakeTopBar()
 	return
 		SNew(SBorder)
 		.BorderImage(FAppStyle::Get().GetBrush("Brushes.Panel"))
-		.Padding(Padding::None)
+		.Padding(NoPadding)
 		[
 			TopBarContent
 		];
@@ -197,10 +199,12 @@ TSharedRef<SWidget> SDataRecoveryTool::MakeRightPanel()
 
 TSharedRef<SWidget> SDataRecoveryTool::MakeDiffSection()
 {
+	const float NoPadding = FDataRecoveryToolStyle::Get().GetFloat("DataRecoveryTool.Padding.None");
+
 	return
 		SNew(SBorder)
 		.BorderImage(FAppStyle::Get().GetBrush("Brushes.Recessed"))
-		.Padding(Padding::None)
+		.Padding(NoPadding)
 		[
 			SAssignNew(DiffView, SInstanceDataObjectFixupTool)
 			.SelectedClassPath(SelectedClassPath)
@@ -350,6 +354,9 @@ FReply SDataRecoveryTool::OnSaveAndApplyClicked() const
 
 TSharedRef<SWidget> SDataRecoveryTool::MakePanel(const FText& InTitle, const TSharedRef<SWidget>& InContent, const TOptional<FOnGetContent>& InMenuContent)
 {
+	const float NoPadding = FDataRecoveryToolStyle::Get().GetFloat("DataRecoveryTool.Padding.None");
+	const float BigPadding = FDataRecoveryToolStyle::Get().GetFloat("DataRecoveryTool.Padding.Big");
+
 	TSharedRef<SHorizontalBox> Header =
 		SNew(SHorizontalBox)
 		+ SHorizontalBox::Slot()
@@ -385,14 +392,14 @@ TSharedRef<SWidget> SDataRecoveryTool::MakePanel(const FText& InTitle, const TSh
 	TSharedRef<SWidget> HeaderTitle =
 		SNew(SBorder)
 		.BorderImage(FAppStyle::Get().GetBrush("Brushes.Panel"))
-		.Padding(Padding::Big)
+		.Padding(BigPadding)
 		[
 			Header
 		];
 
 	return
 		SNew(SBorder)
-		.Padding(Padding::None)
+		.Padding(NoPadding)
 		.BorderImage(FAppStyle::Get().GetBrush("Brushes.Panel"))
 		[
 			SNew(SVerticalBox)
@@ -480,6 +487,8 @@ void SDataRecoveryTool::SetSelectedClass()
 
 TSharedRef<SWidget> SDataRecoveryTool::MakeTEDSTableViewer()
 {
+	const float SmallPadding = FDataRecoveryToolStyle::Get().GetFloat("DataRecoveryTool.Padding.Small");
+
 	using namespace UE::Editor::DataStorage;
 
 	ICoreProvider* Storage = UE::Editor::DataStorage::GetMutableDataStorageFeature<ICoreProvider>(StorageFeatureName);
@@ -535,7 +544,7 @@ TSharedRef<SWidget> SDataRecoveryTool::MakeTEDSTableViewer()
 		SNew(SVerticalBox)
 		+ SVerticalBox::Slot()
 		.AutoHeight()
-		.Padding(Padding::Small)
+		.Padding(SmallPadding)
 		[
 			SearchBoxContainer.ToSharedRef()
 		]
diff --git a/Engine/Plugins/Experimental/InstanceDataObjectFixupTool/Source/Private/DataRecoveryToolStyle.cpp b/Engine/Plugins/Experimental/InstanceDataObjectFixupTool/Source/Private/DataRecoveryToolStyle.cpp
new file mode 100644
--- /dev/null
+++ b/Engine/Plugins/Experimental/InstanceDataObjectFixupTool/Source/Private/DataRecoveryToolStyle.cpp
@@ -0,0 +1,55 @@
+// Copyright Epic Games, Inc. All Rights Reserved.
+
+#include "DataRecoveryToolStyle.h"
+
+#include "Styling/SlateStyleRegistry.h"
+#include "Styling/StarshipCoreStyle.h"
+
+TSharedPtr<FDataRecoveryToolStyle> FDataRecoveryToolStyle::DataRecoveryToolStyle = nullptr;
+
+void FDataRecoveryToolStyle::Register()
+{
+	FSlateStyleRegistry::RegisterSlateStyle(Get());
+}
+
+void FDataRecoveryToolStyle::Unregister()
+{
+	FSlateStyleRegistry::UnRegisterSlateStyle(Get());
+}
+
+void FDataRecoveryToolStyle::Shutdown()
+{
+	Unregister();
+	DataRecoveryToolStyle.Reset();
+}
+
+const FDataRecoveryToolStyle& FDataRecoveryToolStyle::Get()
+{
+	if (!DataRecoveryToolStyle.IsValid())
+	{
+		DataRecoveryToolStyle = MakeShareable(new FDataRecoveryToolStyle());
+	}
+
+	return *DataRecoveryToolStyle;
+}
+
+void FDataRecoveryToolStyle::ReinitializeStyle()
+{
+	Unregister();
+	DataRecoveryToolStyle.Reset();
+	Register();
+}
+
+void FDataRecoveryToolStyle::InitPadding()
+{
+	Set("DataRecoveryTool.Padding.None", 0.f);
+	Set("DataRecoveryTool.Padding.Small", 4.f);
+	Set("DataRecoveryTool.Padding.Normal", 8.f);
+	Set("DataRecoveryTool.Padding.Big", 8.f);
+}
+
+FDataRecoveryToolStyle::FDataRecoveryToolStyle() : FSlateStyleSet("DataRecoveryTool")
+{
+	InitPadding();
+}
+
diff --git a/Engine/Plugins/Experimental/InstanceDataObjectFixupTool/Source/Private/DataRecoveryToolStyle.h b/Engine/Plugins/Experimental/InstanceDataObjectFixupTool/Source/Private/DataRecoveryToolStyle.h
new file mode 100644
--- /dev/null
+++ b/Engine/Plugins/Experimental/InstanceDataObjectFixupTool/Source/Private/DataRecoveryToolStyle.h
@@ -0,0 +1,27 @@
+// Copyright Epic Games, Inc. All Rights Reserved.
+
+#pragma once
+
+#include "CoreMinimal.h"
+#include "Styling/SlateStyle.h"
+
+
+/** Manages the style which provides resources for data recovery tool widgets. */
+class FDataRecoveryToolStyle : public FSlateStyleSet
+{
+public:
+	static void Register();
+	static void Unregister();
+	static void Shutdown();
+
+	/** @return The Slate style set for data recovery tool widgets */
+	static const FDataRecoveryToolStyle& Get();
+
+	static void ReinitializeStyle();
+	void InitPadding();
+
+private:
+	FDataRecoveryToolStyle();
+
+	static TSharedPtr<FDataRecoveryToolStyle> DataRecoveryToolStyle;
+};
diff --git a/Engine/Plugins/Experimental/InstanceDataObjectFixupTool/Source/Private/DataRecoveryToolStyleSheet.h b/Engine/Plugins/Experimental/InstanceDataObjectFixupTool/Source/Private/DataRecoveryToolStyleSheet.h
deleted file mode 100644
diff --git a/Engine/Plugins/Experimental/InstanceDataObjectFixupTool/Source/Private/InstanceDataObjectFixupPanel.cpp b/Engine/Plugins/Experimental/InstanceDataObjectFixupTool/Source/Private/InstanceDataObjectFixupPanel.cpp
--- a/Engine/Plugins/Experimental/InstanceDataObjectFixupTool/Source/Private/InstanceDataObjectFixupPanel.cpp
+++ b/Engine/Plugins/Experimental/InstanceDataObjectFixupTool/Source/Private/InstanceDataObjectFixupPanel.cpp
@@ -4,28 +4,23 @@
 #include "InstanceDataObjectFixupPanel.h"
 
 #include "AsyncDetailViewDiff.h"
-#include "DataRecoveryToolStyleSheet.h"
-#include "DataRecoveryToolUtils.h"
 #include "DetailTreeNode.h"
 #include "Widgets/Layout/LinkableScrollBar.h"
 #include "InstanceDataObjectFixupDetailCustomization.h"
 #include "Misc/StringOutputDevice.h"
 #include "Modules/ModuleManager.h"
 #include "Editor/PropertyEditor/Private/PropertyNode.h"
-#include "TedsAlerts.h"
 #include "UObject/PropertyBagRepository.h"
 #include "UObject/PropertyStateTracking.h"
 
 #include "DataStorage/Features.h"
 #include "Elements/Interfaces/TypedElementDataStorageInterface.h"
-#include "Elements/Interfaces/TypedElementDataStorageCompatibilityInterface.h"
 #include "Serialization/ObjectReader.h"
 #include "Dom/JsonValue.h"
 #include "Dom/JsonObject.h"
 
 #include "UObject/OverriddenPropertySet.h"
 #include "UObject/OverridableManager.h"
-#include "UObject/UObjectThreadContext.h"
 
 #include "UObject/PropertyPathNameTree.h"
 #include "UObject/InstanceDataTransforms.h"
diff --git a/Engine/Plugins/Experimental/InstanceDataObjectFixupTool/Source/Private/InstanceDataObjectFixupTool.cpp b/Engine/Plugins/Experimental/InstanceDataObjectFixupTool/Source/Private/InstanceDataObjectFixupTool.cpp
--- a/Engine/Plugins/Experimental/InstanceDataObjectFixupTool/Source/Private/InstanceDataObjectFixupTool.cpp
+++ b/Engine/Plugins/Experimental/InstanceDataObjectFixupTool/Source/Private/InstanceDataObjectFixupTool.cpp
@@ -2,6 +2,7 @@
 
 #include "InstanceDataObjectFixupTool.h"
 
+#include "DataRecoveryToolStyle.h"
 #include "InstanceDataObjectFixupPanel.h"
 #include "DetailTreeNode.h"
 #include "SlateOptMacros.h"
@@ -9,15 +10,13 @@
 #include "Widgets/Layout/LinkableScrollBar.h"
 #include "Widgets/Input/SButton.h"
 #include "UObject/PropertyBagRepository.h"
-#include "DataRecoveryToolStyleSheet.h"
 #include "DataRecoveryToolUtils.h"
 #include "Widgets/Layout/SScrollBox.h"
 #include "Engine/Engine.h"
 
-#define USE_SNAPSHOT_INSTANCE 0
+#define USE_SNAPSHOT_INSTANCE 1
 
 #define LOCTEXT_NAMESPACE "InstanceDataObjectFixupTool"
-using namespace UE::DataRecoveryTool::StyleSheet;
 
 class FInstanceDataObjectFixupSpecification : public TTreeDiffSpecification<TWeakPtr<FDetailTreeNode>>
 {
@@ -73,14 +72,16 @@ public:
 BEGIN_SLATE_FUNCTION_BUILD_OPTIMIZATION
 void SInstanceDataObjectFixupTool::Construct(const FArguments& InArgs)
 {
+	const float NoPadding = FDataRecoveryToolStyle::Get().GetFloat("DataRecoveryTool.Padding.None");
+
 	SelectedClassPath = InArgs._SelectedClassPath;
 	StagedTransforms = InArgs._StagedTransforms;
 
 	ChildSlot
 	[
 		SAssignNew(Border, SBorder)
-			.BorderImage(FAppStyle::Get().GetBrush("Brushes.Recessed"))
-			.Padding(Padding::None)
+		.BorderImage(FAppStyle::Get().GetBrush("Brushes.Recessed"))
+		.Padding(NoPadding)
 	];
 }
 
@@ -148,9 +149,13 @@ TObjectPtr<UObject> SInstanceDataObjectFixupTool::CreateInstanceDataObjectSnapsh
 
 TSharedRef<SWidget> SInstanceDataObjectFixupTool::MakeDiffViewPanel()
 {
+	const float NoPadding = FDataRecoveryToolStyle::Get().GetFloat("DataRecoveryTool.Padding.None");
+	const float SmallPadding = FDataRecoveryToolStyle::Get().GetFloat("DataRecoveryTool.Padding.Small");
+	const float NormalPadding = FDataRecoveryToolStyle::Get().GetFloat("DataRecoveryTool.Padding.Normal");
+
 	return SNew(SVerticalBox)
 		+ SVerticalBox::Slot()
-		.Padding(Padding::None, Padding::Small)
+		.Padding(NoPadding, SmallPadding)
 		.FillHeight(1.f)
 		[
 			SNew(SScrollBox)
@@ -165,7 +170,7 @@ TSharedRef<SWidget> SInstanceDataObjectFixupTool::MakeDiffViewPanel()
 		.AutoHeight()
 		.HAlign(HAlign_Left)
 		.VAlign(VAlign_Center)
-		.Padding(Padding::Small)
+		.Padding(SmallPadding)
 		[
 			SNew(SHorizontalBox)
 			+ SHorizontalBox::Slot()
@@ -174,7 +179,7 @@ TSharedRef<SWidget> SInstanceDataObjectFixupTool::MakeDiffViewPanel()
 				SNew(SButton)
 				.HAlign(HAlign_Center)
 				.VAlign(VAlign_Center)
-				.ContentPadding(Padding::Normal)
+				.ContentPadding(NormalPadding)
 				.ButtonStyle(FAppStyle::Get(), "SimpleButton")
 				.ToolTipText(LOCTEXT("ResetChangesOnInstance_Tooltip", "Undo all mappings for this class type"))
 				.OnClicked(this, &SInstanceDataObjectFixupTool::OnResetInstanceDataObjectSnapshot)
@@ -190,7 +195,7 @@ TSharedRef<SWidget> SInstanceDataObjectFixupTool::MakeDiffViewPanel()
 					]
 					+ SHorizontalBox::Slot()
 					.AutoWidth()
-					.Padding(Padding::Small, Padding::Small, Padding::Small, Padding::Small)
+					.Padding(SmallPadding)
 					.VAlign(VAlign_Center)
 					[
 						SNew(STextBlock)
@@ -204,7 +209,7 @@ TSharedRef<SWidget> SInstanceDataObjectFixupTool::MakeDiffViewPanel()
 				SNew(SButton)
 				.HAlign(HAlign_Center)
 				.VAlign(VAlign_Center)
-				.ContentPadding(Padding::Normal)
+				.ContentPadding(NormalPadding)
 				.ButtonStyle(FAppStyle::Get(), "SimpleButton")
 				.ToolTipText(LOCTEXT("AutoMarkForDeletion_Tooltip", "Mark remaining conflicts for deletion"))
 				.OnClicked(this, &SInstanceDataObjectFixupTool::OnAutoMarkForDeletion)
@@ -220,7 +225,7 @@ TSharedRef<SWidget> SInstanceDataObjectFixupTool::MakeDiffViewPanel()
 					]
 					+ SHorizontalBox::Slot()
 					.AutoWidth()
-					.Padding(Padding::Small, Padding::Small, Padding::Small, Padding::Small)
+					.Padding(SmallPadding)
 					.VAlign(VAlign_Center)
 					[
 						SNew(STextBlock)
diff --git a/Engine/Plugins/Experimental/InstanceDataObjectFixupTool/Source/Private/InstanceDataObjectFixupToolModule.cpp b/Engine/Plugins/Experimental/InstanceDataObjectFixupTool/Source/Private/InstanceDataObjectFixupToolModule.cpp
--- a/Engine/Plugins/Experimental/InstanceDataObjectFixupTool/Source/Private/InstanceDataObjectFixupToolModule.cpp
+++ b/Engine/Plugins/Experimental/InstanceDataObjectFixupTool/Source/Private/InstanceDataObjectFixupToolModule.cpp
@@ -3,7 +3,6 @@
 #include "InstanceDataObjectFixupToolModule.h"
 
 #include "DataRecoveryTool.h"
-#include "DataRecoveryToolStyleSheet.h"
 #include "WorkspaceMenuStructure.h"
 #include "WorkspaceMenuStructureModule.h"
 #include "Elements/Framework/TypedElementRegistry.h"
@@ -13,6 +12,7 @@
 
 #define LOCTEXT_NAMESPACE "InstanceDataObjectFixupToolModule"
 
+FName DataRecoveryToolDialogName = TEXT("DataRecoveryToolDialog");
 
 void FInstanceDataObjectFixupToolModule::StartupModule()
 {
@@ -22,10 +22,8 @@ void FInstanceDataObjectFixupToolModule::StartupModule()
 
 void FInstanceDataObjectFixupToolModule::RegisterTabSpawners(const TOptional<FTopLevelAssetPath>& ClassPath) const
 {
-	using namespace UE::DataRecoveryTool::StyleSheet;
-
 	FGlobalTabmanager::Get()->RegisterNomadTabSpawner(
-		Dialog::DataRecoveryToolDialogName,
+		DataRecoveryToolDialogName,
 		FOnSpawnTab::CreateLambda([this, ClassPath](const FSpawnTabArgs& Args)
 			{
 				return CreateDataRecoveryToolTab(ClassPath);
@@ -33,16 +31,14 @@ void FInstanceDataObjectFixupToolModule::RegisterTabSpawners(const TOptional<FTo
 		.SetGroup(WorkspaceMenu::GetMenuStructure().GetToolsCategory())
 		.SetDisplayName(LOCTEXT("DataRecoveryTool_DisplayName", "Data Recovery Tool"))
 		.SetTooltipText(LOCTEXT("DataRecoveryTool_DisplayNameToolTip", "Opens the Data recovery Tool"))
-		.SetIcon(FSlateIcon(FAppStyle::GetAppStyleSetName(), Tab::Icon));
+		.SetIcon(FSlateIcon(FAppStyle::GetAppStyleSetName(), "Icons.Convert"));
 }
 
 void FInstanceDataObjectFixupToolModule::UnregisterTabSpawners() const
 {
-	using namespace UE::DataRecoveryTool::StyleSheet;
-
 	if (FSlateApplication::IsInitialized())
 	{
-		FGlobalTabmanager::Get()->UnregisterNomadTabSpawner(Dialog::DataRecoveryToolDialogName);
+		FGlobalTabmanager::Get()->UnregisterNomadTabSpawner(DataRecoveryToolDialogName);
 	}
 }
 
@@ -69,11 +65,9 @@ TSharedRef<SDockTab> FInstanceDataObjectFixupToolModule::CreateDataRecoveryToolT
 
 void FInstanceDataObjectFixupToolModule::CreateDataRecoveryToolDialog(const TOptional<FTopLevelAssetPath>& ClassPath) const
 {
-	using namespace UE::DataRecoveryTool::StyleSheet;
-
 	RegisterTabSpawners(ClassPath);
 
-	if (const TSharedPtr<SDockTab> DockTab = FGlobalTabmanager::Get()->TryInvokeTab(FTabId(Dialog::DataRecoveryToolDialogName)))
+	if (const TSharedPtr<SDockTab> DockTab = FGlobalTabmanager::Get()->TryInvokeTab(DataRecoveryToolDialogName))
 	{
 		DockTab->DrawAttention();
 	}
