This is the final CL that marks the Data Recovery Tool "feature ready"! We technically still haven't added funtionality for remapping placeholder classes but all property remapping features are complete.

[Feature] DRT type conversion options
[Added] type conversion and type conversion & renaming options.
[Added] ] IDO.DataRecoveryTool.TypeConversions CVar
[Changed] Now using snapshot instance
[BugFix] Checking if instance still requires fixup before counting it as has loose properties

#Synced-CL 48303336

diff --git a/Engine/Plugins/Experimental/InstanceDataObjectFixupTool/Source/Private/DataRecoveryTool.cpp b/Engine/Plugins/Experimental/InstanceDataObjectFixupTool/Source/Private/DataRecoveryTool.cpp
--- a/Engine/Plugins/Experimental/InstanceDataObjectFixupTool/Source/Private/DataRecoveryTool.cpp
+++ b/Engine/Plugins/Experimental/InstanceDataObjectFixupTool/Source/Private/DataRecoveryTool.cpp
@@ -241,7 +241,7 @@ FReply SDataRecoveryTool::OnApplyClicked() const
 
 			TArray<UObject*> AllInstancesOfClasses;
 
-			uint8 InstanceWithUnknownPropertiesCount = 0;
+			uint32 InstanceWithUnknownPropertiesCount = 0;
 
 			GetObjectsOfClass(Class, AllInstancesOfClasses);
 
@@ -253,7 +253,10 @@ FReply SDataRecoveryTool::OnApplyClicked() const
 					{
 						ApplyTransforms(StagedTransforms, Ido, Instance);
 
-						++InstanceWithUnknownPropertiesCount;
+						if (UE::FPropertyBagRepository::Get().RequiresFixup(Instance))
+						{
+							++InstanceWithUnknownPropertiesCount;
+						}
 					}
 				}
 			}
@@ -317,7 +320,7 @@ FReply SDataRecoveryTool::OnSaveAndApplyClicked() const
 
 			TArray<UObject*> AllInstancesOfClasses;
 
-			uint8 InstanceWithUnknownPropertiesCount = 0;
+			uint32 InstanceWithUnknownPropertiesCount = 0;
 
 			// Could not use PatchTransformSet inside ForEachObjectOfClass lambda, so we have to get the objects first
 			GetObjectsOfClass(Class, AllInstancesOfClasses);
@@ -330,7 +333,10 @@ FReply SDataRecoveryTool::OnSaveAndApplyClicked() const
 					{
 						UE::FInstanceDataTransforms::Get().PatchInstanceDataObject(Ido, Instance);
 
-						++InstanceWithUnknownPropertiesCount;
+						if (UE::FPropertyBagRepository::Get().RequiresFixup(Instance))
+						{
+							++InstanceWithUnknownPropertiesCount;
+						}
 					}
 				}
 			}
diff --git a/Engine/Plugins/Experimental/InstanceDataObjectFixupTool/Source/Private/InstanceDataObjectFixupDetailCustomization.cpp b/Engine/Plugins/Experimental/InstanceDataObjectFixupTool/Source/Private/InstanceDataObjectFixupDetailCustomization.cpp
--- a/Engine/Plugins/Experimental/InstanceDataObjectFixupTool/Source/Private/InstanceDataObjectFixupDetailCustomization.cpp
+++ b/Engine/Plugins/Experimental/InstanceDataObjectFixupTool/Source/Private/InstanceDataObjectFixupDetailCustomization.cpp
@@ -10,8 +10,17 @@
 #include "Widgets/Layout/SWidgetSwitcher.h"
 #include "Widgets/Input/SComboButton.h"
 #include "Framework/MultiBox/MultiBoxBuilder.h"
+#include "Serialization/ObjectReader.h"
 #include "UObject/PropertyBagRepository.h"
 
+
+static bool bEnableTypeConversions = true;
+
+static FAutoConsoleVariableRef CVarEnableTypeConversions(
+	TEXT("IDO.DataRecoveryTool.TypeConversions"),
+	bEnableTypeConversions,
+	TEXT("Enable the menu entry to convert types in the data recovery tool"));
+
 #define LOCTEXT_NAMESPACE "InstanceDataObjectFixupDetails"
 
 /////////////////////////////////////////////////////////////////
@@ -187,6 +196,11 @@ static bool DoesRedirectRespectContainerIndices(const FPropertyPath& Src, const
 	return true;
 }
 
+static bool PropertiesAreMatchingTypes(const FProperty* A, const FProperty* B)
+{
+	return UE::FPropertyTypeName(A) == UE::FPropertyTypeName(B);
+}
+
 TSharedRef<SWidget> FInstanceDataObjectNameWidgetOverride::GeneratePropertyRedirectMenu(FPropertyPath Path) const
 {
 	FMenuBuilder MenuBuilder(true, nullptr);
@@ -226,7 +240,7 @@ TSharedRef<SWidget> FInstanceDataObjectNameWidgetOverride::GeneratePropertyRedir
 			const FProperty* OptionProperty = Option.GetLeafMostProperty().Property.Get();
 			if (ThisProperty->GetFName() == OptionProperty->GetFName())
 			{
-				if (OptionProperty->SameType(ThisProperty))
+				if (PropertiesAreMatchingTypes(OptionProperty, ThisProperty))
 				{
 					FText DisplayName = GetCleanVersePath(Option);
 					FText Tooltip = FText::Format(LOCTEXT("MovePropertyTooltip", "Move property to '{0}'"), DisplayName);
@@ -253,7 +267,7 @@ TSharedRef<SWidget> FInstanceDataObjectNameWidgetOverride::GeneratePropertyRedir
 			{
 				continue; // handled in the "move" category
 			}
-			if (OptionProperty->SameType(ThisProperty))
+			if (PropertiesAreMatchingTypes(OptionProperty, ThisProperty))
 			{
 				FText DisplayName = GetCleanVersePath(Option);
 				FText Tooltip = FText::Format(LOCTEXT("RenamePropertyTooltip", "Rename property to '{0}'"), DisplayName);
@@ -265,75 +279,61 @@ TSharedRef<SWidget> FInstanceDataObjectNameWidgetOverride::GeneratePropertyRedir
 	}
 	MenuBuilder.EndSection();
 
-	//MenuBuilder.BeginSection(NAME_None, LOCTEXT("ChangeToDifferingTypeConversion", "Convert Type"));
-	//{
-	//	for (const FPropertyPath& Option : RedirectOptions)
-	//	{
-	//		const FProperty* ThisProperty = Path.GetLeafMostProperty().Property.Get();
-	//		const FProperty* OptionProperty = Option.GetLeafMostProperty().Property.Get();
-	//		if (OptionProperty->SameType(ThisProperty))
-	//		{
-	//			continue; // same type handled above
-	//		}
-	//		if (ThisProperty->GetFName() != OptionProperty->GetFName())
-	//		{
-	//			continue; // renames to handled below
-	//		}
-	//		if (FInstanceDataObjectFixupPanel::FTypeConverter Converter = Panel->CreateTypeConverter(Path, Option))
-	//		{
-	//			FText DisplayName = GetCleanVersePath(Option);
-	//			FText TypeName = FText::FromName(Option.GetLeafMostProperty().Property->GetID());
-	//			FText Warning = Converter.GetWarning();
-	//			FText Tooltip = FText::Format(LOCTEXT("ConvertTypeTooltip", "Change type to {0}"), TypeName);
-	//			if (!Warning.IsEmpty())
-	//			{
-	//				DisplayName = FText::Format(LOCTEXT("ConvertTypeDisplayNameWithWarning", "⚠{0}"), DisplayName);
-	//				Tooltip = FText::Format(LOCTEXT("ConvertTypeTooltipWithWarning", "⚠{0}\n Warning: {1}"), Tooltip, Warning);
-	//			}
-	//			MenuBuilder.AddMenuEntry(DisplayName, Tooltip, FSlateIcon()
-	//			, FUIAction(FExecuteAction::CreateSP(Panel.Get(), &FInstanceDataObjectFixupPanel::OnRedirectProperty, Path, Option, Converter))
-	//			, NAME_None
-	//			, EUserInterfaceActionType::RadioButton);
-	//		}
-	//	}
-	//}
-	//MenuBuilder.EndSection();
+	if (bEnableTypeConversions)
+	{
+		MenuBuilder.BeginSection(NAME_None, LOCTEXT("ChangeToDifferingTypeConversion", "Convert Type"));
+		{
+			for (const FPropertyPath& Option : RedirectOptions)
+			{
+				const FProperty* ThisProperty = Path.GetLeafMostProperty().Property.Get();
+				const FProperty* OptionProperty = Option.GetLeafMostProperty().Property.Get();
+				if (PropertiesAreMatchingTypes(OptionProperty, ThisProperty))
+				{
+					continue; // same type handled above
+				}
+				if (ThisProperty->GetFName() != OptionProperty->GetFName())
+				{
+					continue; // renames to handled below
+				}
+				if (Panel->CanConvert(Path, const_cast<FProperty*>(OptionProperty)))
+				{
+					FText DisplayName = GetCleanVersePath(Option);
+					FText Tooltip = FText::Format(LOCTEXT("ConvertPropertyTooltip", "Convert property to '{0}'"), DisplayName);
+					MenuBuilder.AddMenuEntry(DisplayName, Tooltip, FSlateIcon(FAppStyle::GetAppStyleSetName(), FName("Icons.Convert"))
+						, FUIAction(FExecuteAction::CreateSP(Panel.Get(), &FInstanceDataObjectFixupPanel::OnRedirectProperty, Path, Option))
+						, NAME_None);
+				}
+			}
+		}
+		MenuBuilder.EndSection();
+
+		MenuBuilder.BeginSection(NAME_None, LOCTEXT("RenameToDifferingTypeConversion", "Convert Type and Rename"));
+		{
+			for (const FPropertyPath& Option : RedirectOptions)
+			{
+				const FProperty* ThisProperty = Path.GetLeafMostProperty().Property.Get();
+				const FProperty* OptionProperty = Option.GetLeafMostProperty().Property.Get();
+				if (PropertiesAreMatchingTypes(OptionProperty, ThisProperty))
+				{
+					continue; // same type handled above
+				}
+				if (ThisProperty->GetFName() == OptionProperty->GetFName())
+				{
+					continue; // renames to handled below
+				}
+				if (Panel->CanConvert(Path, const_cast<FProperty*>(OptionProperty)))
+				{
+					FText DisplayName = GetCleanVersePath(Option);
+					FText Tooltip = FText::Format(LOCTEXT("ConvertPropertyTooltip", "Convert property to '{0}'"), DisplayName);
+					MenuBuilder.AddMenuEntry(DisplayName, Tooltip, FSlateIcon(FAppStyle::GetAppStyleSetName(), FName("Icons.Convert"))
+						, FUIAction(FExecuteAction::CreateSP(Panel.Get(), &FInstanceDataObjectFixupPanel::OnRedirectProperty, Path, Option))
+						, NAME_None);
+				}
+			}
+		}
+		MenuBuilder.EndSection();
+	}
 
-	//MenuBuilder.BeginSection(NAME_None, LOCTEXT("RenameToDifferingTypeConversion", "Convert Type and Rename"));
-	//{
-	//	for (const FPropertyPath& Option : RedirectOptions)
-	//	{
-	//		const FProperty* ThisProperty = Path.GetLeafMostProperty().Property.Get();
-	//		const FProperty* OptionProperty = Option.GetLeafMostProperty().Property.Get();
-	//		if (OptionProperty->SameType(ThisProperty))
-	//		{
-	//			continue; // same type handled above
-	//		}
-	//		if (ThisProperty->GetFName() == OptionProperty->GetFName())
-	//		{
-	//			continue; // renames to handled below
-	//		}
-	//		if (FInstanceDataObjectFixupPanel::FTypeConverter Converter = Panel->CreateTypeConverter(Path, Option))
-	//		{
-	//			FText DisplayName = GetCleanVersePath(Option);
-	//			FText PropDisplayName = Option.GetLeafMostProperty().Property->GetDisplayNameText();
-	//			FText TypeName = FText::FromName(Option.GetLeafMostProperty().Property->GetID());
-	//			FText Warning = Converter.GetWarning();
-	//			FText Tooltip = FText::Format(LOCTEXT("ConvertTypeAndRenameTooltip", "Change type to {0} and rename to {1}"), TypeName, PropDisplayName);
-	//			if (!Warning.IsEmpty())
-	//			{
-	//				DisplayName = FText::Format(LOCTEXT("ConvertTypeAndRenameDisplayNameWithWarning", "⚠{0}"), DisplayName);
-	//				Tooltip = FText::Format(LOCTEXT("ConvertTypeAndRenameTooltipWithWarning", "⚠{0}\n Warning: {1}"), Tooltip, Warning);
-	//			}
-	//			MenuBuilder.AddMenuEntry(DisplayName, Tooltip, FSlateIcon()
-	//			, FUIAction(FExecuteAction::CreateSP(Panel.Get(), &FInstanceDataObjectFixupPanel::OnRedirectProperty, Path, Option, Converter))
-	//			, NAME_None
-	//			, EUserInterfaceActionType::RadioButton);
-	//		}
-	//	}
-	//}
-	//MenuBuilder.EndSection();
-	
 	return MenuBuilder.MakeWidget();
 }
 
diff --git a/Engine/Plugins/Experimental/InstanceDataObjectFixupTool/Source/Private/InstanceDataObjectFixupPanel.cpp b/Engine/Plugins/Experimental/InstanceDataObjectFixupTool/Source/Private/InstanceDataObjectFixupPanel.cpp
--- a/Engine/Plugins/Experimental/InstanceDataObjectFixupTool/Source/Private/InstanceDataObjectFixupPanel.cpp
+++ b/Engine/Plugins/Experimental/InstanceDataObjectFixupTool/Source/Private/InstanceDataObjectFixupPanel.cpp
@@ -24,11 +24,115 @@
 
 #include "UObject/PropertyPathNameTree.h"
 #include "UObject/InstanceDataTransforms.h"
+#include "UObject/UObjectThreadContext.h"
 
 #define LOCTEXT_NAMESPACE "InstanceDataObjectFixupPanel"
 
 static const FName NAME_IsLooseMetadata(TEXT("IsLoose"));
 
+static void* ResolvePath(const FPropertyPath& Path, void* Value)
+{
+	for (int32 PathIndex = 0; PathIndex < Path.GetNumProperties(); ++PathIndex)
+	{
+		const FPropertyInfo& PropertyInfo = Path.GetPropertyInfo(PathIndex);
+		const FProperty* Property = PropertyInfo.Property.Get();
+		if (!Property)
+		{
+			return nullptr;
+		}
+
+		Value = Property->ContainerPtrToValuePtr<void>(Value, PropertyInfo.ArrayIndex != INDEX_NONE ? PropertyInfo.ArrayIndex : 0);
+
+		if (const FObjectProperty* AsObjectProperty = CastField<FObjectProperty>(Property))
+		{
+			UObject* Object = AsObjectProperty->GetObjectPropertyValue(Value);
+			UE::FPropertyBagRepository& PropertyBagRepository = UE::FPropertyBagRepository::Get();
+			if (UObject* Found = PropertyBagRepository.FindInstanceDataObject(Object))
+			{
+				Object = Found;
+			}
+			Value = Object;
+		}
+		else if (PathIndex + 1 < Path.GetNumProperties())
+		{
+			if (const FArrayProperty* AsArrayProperty = CastField<FArrayProperty>(Property))
+			{
+				FScriptArrayHelper Helper(AsArrayProperty, Value);
+				Value = Helper.GetElementPtr(Path.GetPropertyInfo(++PathIndex).ArrayIndex);
+			}
+			if (const FSetProperty* AsSetProperty = CastField<FSetProperty>(Property))
+			{
+				FScriptSetHelper Helper(AsSetProperty, Value);
+				Value = Helper.FindNthElementPtr(Path.GetPropertyInfo(++PathIndex).ArrayIndex);
+			}
+			if (const FMapProperty* AsMapProperty = CastField<FMapProperty>(Property))
+			{
+				FScriptMapHelper Helper(AsMapProperty, Value);
+				Value = Helper.FindNthValuePtr(Path.GetPropertyInfo(++PathIndex).ArrayIndex);
+			}
+		}
+	}
+
+	return Value;
+}
+
+
+bool FInstanceDataObjectFixupPanel::CanConvert(const FPropertyPath& SourcePropertyPath, FProperty* DestinationProperty)
+{
+	FProperty* SourceProperty = SourcePropertyPath.GetLeafMostProperty().Property.Get();
+
+	if (SourceProperty->GetID() == DestinationProperty->GetID())
+	{
+		return true;
+	}
+
+	const void* SourceData = ResolvePath(SourcePropertyPath, InstanceDataObject);
+
+	if (SourceData == nullptr)
+	{
+		return false;
+	}
+
+	FUObjectSerializeContext* SerializeContext = FUObjectThreadContext::Get().GetSerializeContext();
+	TGuardValue<bool> ScopedImpersonateProperties(SerializeContext->bImpersonateProperties, true);
+
+	TArray<uint8, TInlineAllocator<64>> Buffer;
+	TMemoryWriterBase<TInlineAllocator<64>> MemoryWriter(Buffer);
+	FStructuredArchiveFromArchive StructuredWriter(MemoryWriter);
+
+	SourceProperty->SerializeItem(StructuredWriter.GetSlot(), (uint8*)SourceData);
+	FMemoryReaderView MemoryReader(Buffer);
+	FStructuredArchiveFromArchive StructuredReader(MemoryReader);
+
+
+	TArray<uint8, TInlineAllocator<64>> SourceToDest;
+	SourceToDest.SetNumUninitialized(DestinationProperty->GetElementSize());
+	DestinationProperty->InitializeValue(SourceToDest.GetData());
+	void* DestinationContainer = SourceToDest.GetData();
+
+	// todo: handle static arrays
+	FPropertyTag SourceTag(SourceProperty, 0, (uint8*)SourceData);
+
+	bool bResult = false;
+	switch (DestinationProperty->ConvertFromType(SourceTag, StructuredReader.GetSlot(), (uint8*)DestinationContainer, SourceProperty->GetOwnerStruct(), nullptr))
+	{
+	case EConvertFromTypeResult::UseSerializeItem:
+		break;
+	case EConvertFromTypeResult::Serialized:
+		bResult = true;
+		break;
+	case EConvertFromTypeResult::CannotConvert:
+		break;
+	case EConvertFromTypeResult::Converted:
+		bResult = true;
+		break;
+	}
+
+	DestinationProperty->DestroyValue(SourceToDest.GetData());
+
+	return bResult;
+}
+
 static bool IsPathOverridden(UObject* Object, const FPropertyPath& Path)
 {
 	FPropertyVisitorPath VisitorPath;
diff --git a/Engine/Plugins/Experimental/InstanceDataObjectFixupTool/Source/Private/InstanceDataObjectFixupPanel.h b/Engine/Plugins/Experimental/InstanceDataObjectFixupTool/Source/Private/InstanceDataObjectFixupPanel.h
--- a/Engine/Plugins/Experimental/InstanceDataObjectFixupTool/Source/Private/InstanceDataObjectFixupPanel.h
+++ b/Engine/Plugins/Experimental/InstanceDataObjectFixupTool/Source/Private/InstanceDataObjectFixupPanel.h
@@ -26,7 +26,7 @@ public:
 		// display only properties found in latest version of the class but allow value edits
 		DefaultRightPanel = HideLooseProperties,
 	};
-	
+
 	FInstanceDataObjectFixupPanel(const TObjectPtr<UObject>& InstanceDataObject, EViewFlags ViewFlags, const TWeakPtr<TMap<FTopLevelAssetPath, UE::FInstanceDataTransformSet>>& StagedTransforms);
 
 	void RefreshDetailsView();
@@ -38,7 +38,7 @@ public:
 	TSharedPtr<FAsyncDetailViewDiff> GetDiffAgainstRight() const;
 
 	bool AreAllConflictsRedirected() const;
-	
+
 	void MarkForDelete(const FPropertyPath& Path);
 	// pass-by-copy version for delegates. Use MarkForDelete when possible
 	void OnMarkForDelete(FPropertyPath Path);
@@ -62,6 +62,8 @@ private:
 
 	void RedirectProperty(const FPropertyPath& From, const FPropertyPath& To);
 
+	bool CanConvert(const FPropertyPath& SourcePropertyPath, FProperty* DestinationProperty);
+
 	// pass-by-copy version for delegates. Use RedirectProperty when possible
 	void OnRedirectProperty(FPropertyPath From, FPropertyPath To);
 
