[CrashFix] Outliner Warning crashing when opening DRT

[Fix] Deffered the launch of the Data Recovery Tool after the function that calls ForEachObjectWithOuter()
[Fix] Removed the possibility of the DRT to be launched if the cvar for it is disabled

#Synced-CL 48032063

diff --git a/Engine/Plugins/Experimental/InstanceDataObjectFixupTool/Source/Private/InstanceDataObjectFixupToolModule.cpp b/Engine/Plugins/Experimental/InstanceDataObjectFixupTool/Source/Private/InstanceDataObjectFixupToolModule.cpp
--- a/Engine/Plugins/Experimental/InstanceDataObjectFixupTool/Source/Private/InstanceDataObjectFixupToolModule.cpp
+++ b/Engine/Plugins/Experimental/InstanceDataObjectFixupTool/Source/Private/InstanceDataObjectFixupToolModule.cpp
@@ -63,8 +63,23 @@ TSharedRef<SDockTab> FInstanceDataObjectFixupToolModule::CreateDataRecoveryToolT
 	return DockTab;
 }
 
+bool FInstanceDataObjectFixupToolModule::IsDataRecoveryToolEnabled() const
+{
+	if (const IConsoleVariable* CVarEnableDataRecoveryTool = IConsoleManager::Get().FindConsoleVariable(TEXT("IDO.EnableDataRecoveryTool")))
+	{
+		return CVarEnableDataRecoveryTool->GetBool();
+	}
+
+	return false;
+}
+
 void FInstanceDataObjectFixupToolModule::CreateDataRecoveryToolDialog(const TOptional<FTopLevelAssetPath>& ClassPath) const
 {
+	if (!IsDataRecoveryToolEnabled())
+	{
+		return;
+	}
+
 	RegisterTabSpawners(ClassPath);
 
 	if (const TSharedPtr<SDockTab> DockTab = FGlobalTabmanager::Get()->TryInvokeTab(DataRecoveryToolDialogName))
diff --git a/Engine/Plugins/Experimental/InstanceDataObjectFixupTool/Source/Private/InstanceDataObjectFixupToolTedsQueries.cpp b/Engine/Plugins/Experimental/InstanceDataObjectFixupTool/Source/Private/InstanceDataObjectFixupToolTedsQueries.cpp
--- a/Engine/Plugins/Experimental/InstanceDataObjectFixupTool/Source/Private/InstanceDataObjectFixupToolTedsQueries.cpp
+++ b/Engine/Plugins/Experimental/InstanceDataObjectFixupTool/Source/Private/InstanceDataObjectFixupToolTedsQueries.cpp
@@ -71,12 +71,17 @@ void UInstanceDataObjectFixupToolTedsQueryFactory::ShowFixUpTool(UE::Editor::Dat
 		if (bRecurseIntoObject)
 		{
 			UObject* Owner = ObjectColumn->Object.Get();
+			TOptional<FTopLevelAssetPath> ClassPath;
+
 			UE::FPropertyBagRepository::Get().FindNestedInstanceDataObject(Owner, true,
-				[Owner](UObject* NestedObject)
+				[Owner, &ClassPath](UObject* NestedObject)
 				{
-					FTopLevelAssetPath ClassPath = Owner->GetClass()->GetClassPathName();
-					FInstanceDataObjectFixupToolModule::Get().CreateDataRecoveryToolDialog(ClassPath);
+					ClassPath = Owner->GetClass()->GetClassPathName();
 				});
+
+
+			FInstanceDataObjectFixupToolModule::Get().CreateDataRecoveryToolDialog(ClassPath);
+			
 		}
 		else
 		{
diff --git a/Engine/Plugins/Experimental/InstanceDataObjectFixupTool/Source/Public/InstanceDataObjectFixupToolModule.h b/Engine/Plugins/Experimental/InstanceDataObjectFixupTool/Source/Public/InstanceDataObjectFixupToolModule.h
--- a/Engine/Plugins/Experimental/InstanceDataObjectFixupTool/Source/Public/InstanceDataObjectFixupToolModule.h
+++ b/Engine/Plugins/Experimental/InstanceDataObjectFixupTool/Source/Public/InstanceDataObjectFixupToolModule.h
@@ -17,12 +17,15 @@ public:
 		return FModuleManager::LoadModuleChecked<FInstanceDataObjectFixupToolModule>("InstanceDataObjectFixupTool");
 	}
 
-	UE_API TSharedRef<SDockTab> CreateDataRecoveryToolTab(const TOptional<FTopLevelAssetPath>& ClassPath = {}) const;
-
 	UE_API void CreateDataRecoveryToolDialog(const TOptional<FTopLevelAssetPath>& ClassPath) const;
 
 	UE_API void RegisterTabSpawners(const TOptional<FTopLevelAssetPath>& ClassPath = {}) const;
 	UE_API void UnregisterTabSpawners() const;
+
+private:
+	TSharedRef<SDockTab> CreateDataRecoveryToolTab(const TOptional<FTopLevelAssetPath>& ClassPath = {}) const;
+
+	bool IsDataRecoveryToolEnabled() const;
 };
 
 #undef UE_API
