[TEDS Searching] Move search in the DRT to use new Query Search

#Synced-CL 49995863

diff --git a/Engine/Plugins/Experimental/InstanceDataObjectFixupTool/Source/Private/DataRecoveryTool.cpp b/Engine/Plugins/Experimental/InstanceDataObjectFixupTool/Source/Private/DataRecoveryTool.cpp
--- a/Engine/Plugins/Experimental/InstanceDataObjectFixupTool/Source/Private/DataRecoveryTool.cpp
+++ b/Engine/Plugins/Experimental/InstanceDataObjectFixupTool/Source/Private/DataRecoveryTool.cpp
@@ -8,6 +8,7 @@
 #include "Elements/Columns/TedsTypeInfoColumns.h"
 #include "Elements/Columns/TypedElementCompatibilityColumns.h"
 #include "Elements/Columns/TypedElementHiearchyColumns.h"
+#include "Elements/Columns/TypedElementLabelColumns.h"
 #include "Elements/Columns/TypedElementTypeInfoColumns.h"
 #include "Elements/Framework/TypedElementQueryBuilder.h"
 #include "Elements/Interfaces/TypedElementDataStorageInterface.h"
@@ -509,38 +510,43 @@ TSharedRef<SWidget> SDataRecoveryTool::MakeTEDSTableViewer()
 	using namespace UE::Editor::DataStorage::Queries;
 
 	TSharedRef<QueryStack::FQueryNode> ReferenceQueryNode =
-		MakeShared<QueryStack::FQueryNode>(*Storage, Select().Where().All<FHasUnknownPropertiesTag, FTypedElementClassTypeInfoColumn>().Compile());
-
-	ReferenceQueryResultsNode = MakeShared<QueryStack::FRowQueryResultsNode>(*Storage,
-		ReferenceQueryNode,
-		QueryStack::FRowQueryResultsNode::ESyncActions::ForceRefreshOnUpdate);
+		MakeShared<QueryStack::FQueryNode>(*Storage,
+			Select()
+				.ReadOnly<FTypedElementLabelColumn>()
+			.Where()
+				.All<FHasUnknownPropertiesTag, FTypedElementClassTypeInfoColumn>().Compile());
 
 	TSharedPtr<QueryStack::IRowNode> SearchNode;
 	TSharedPtr<SBox> SearchBoxContainer = SNew(SBox);
 
 	SearchBoxContainer->SetContent(SNew(UE::Editor::DataStorage::STedsSearchBox)
-		.InSearchableRowNode(ReferenceQueryResultsNode)
+		.InSearchableQueryNode(ReferenceQueryNode)
+		.InSearchableQueryFlags(QueryStack::FRowQueryResultsNode::ESyncActions::ForceRefreshOnUpdate)
 		.OutSearchNode(&SearchNode));
 
-
 	// If there is no SearchNode available, either disabled or there is an internal error, assign it to be the 
 	// ReferenceQueryResultsNode so that it won't be null
 	if (!SearchNode)
 	{
-		SearchNode = ReferenceQueryResultsNode;
+		ReferenceQueryResultsNode = MakeShared<QueryStack::FRowQueryResultsNode>(
+			*Storage,
+			ReferenceQueryNode,
+			QueryStack::FRowQueryResultsNode::ESyncActions::ForceRefreshOnUpdate);
+	}
+	else
+	{
+		ReferenceQueryResultsNode = SearchNode;
 	}
 
 	// Using the Teds Outliner Purpose for matching customization for now
 	IUiProvider::FPurposeID PurposeId = IUiProvider::FPurposeInfo("SceneOutliner", "RowLabel", NAME_None).GeneratePurposeID();
 
 	ClassesTableViewer = SNew(STedsTableViewer)
-		.QueryStack(SearchNode)
+		.QueryStack(ReferenceQueryResultsNode)
 		.Columns({ FTypedElementClassTypeInfoColumn::StaticStruct(), FTedsAlertColumn::StaticStruct() })
 		.ListSelectionMode(ESelectionMode::Single)
 		.OnSelectionChanged(this, &SDataRecoveryTool::OnListSelectionChanged);
 
-
-
 	if (SelectedClassPath.IsValid())
 	{
 		SetSelectedClass();
diff --git a/Engine/Plugins/Experimental/InstanceDataObjectFixupTool/Source/Public/DataRecoveryTool.h b/Engine/Plugins/Experimental/InstanceDataObjectFixupTool/Source/Public/DataRecoveryTool.h
--- a/Engine/Plugins/Experimental/InstanceDataObjectFixupTool/Source/Public/DataRecoveryTool.h
+++ b/Engine/Plugins/Experimental/InstanceDataObjectFixupTool/Source/Public/DataRecoveryTool.h
@@ -69,7 +69,7 @@ private:
 
 	TSharedPtr<UE::Editor::DataStorage::STedsTableViewer> ClassesTableViewer;
 
-	TSharedPtr<UE::Editor::DataStorage::QueryStack::FRowQueryResultsNode> ReferenceQueryResultsNode;
+	TSharedPtr<UE::Editor::DataStorage::QueryStack::IRowNode> ReferenceQueryResultsNode;
 
 	TSharedPtr<SInstanceDataObjectFixupTool> DiffView;
 };
